IoStore trace plugin for Unreal Insights

#rnx

diff --git a/Engine/Plugins/IoStoreInsights/IoStoreInsights.uplugin b/Engine/Plugins/IoStoreInsights/IoStoreInsights.uplugin
new file mode 100644
--- /dev/null
+++ b/Engine/Plugins/IoStoreInsights/IoStoreInsights.uplugin
@@ -0,0 +1,28 @@
+{
+	"FileVersion": 3,
+	"Version": 1,
+	"VersionName": "1.0",
+	"FriendlyName": "IoStore Insights",
+	"Description": "Allows capturing IoStore activity via Unreal Insights",
+	"Category": "Insights",
+	"CreatedBy": "Epic Games, Inc.",
+	"CreatedByURL": "https://epicgames.com",
+	"DocsURL": "",
+	"MarketplaceURL": "",
+	"SupportURL": "",
+	"EnabledByDefault": true,
+	"CanContainContent": false,
+	"IsBetaVersion": false,
+	"Installed": false,
+	"SupportedPrograms": [ "UnrealInsights" ],
+	"Modules": [
+		{
+			"Name": "IoStoreInsights",
+			"Type": "EditorAndProgram",
+			"LoadingPhase": "Default",
+			"ProgramAllowList": [
+				"UnrealInsights"
+			]
+		}
+	]
+}
diff --git a/Engine/Plugins/IoStoreInsights/Source/IoStoreInsights/IoStoreInsights.Build.cs b/Engine/Plugins/IoStoreInsights/Source/IoStoreInsights/IoStoreInsights.Build.cs
new file mode 100644
--- /dev/null
+++ b/Engine/Plugins/IoStoreInsights/Source/IoStoreInsights/IoStoreInsights.Build.cs
@@ -0,0 +1,24 @@
+// Copyright Epic Games, Inc. All Rights Reserved.
+
+namespace UnrealBuildTool.Rules
+{
+	public class IoStoreInsights : ModuleRules
+	{
+		public IoStoreInsights(ReadOnlyTargetRules Target) : base(Target)
+		{
+			PrivateDependencyModuleNames.AddRange(
+			new string[]
+			{
+				"Core",
+				"SlateCore",
+				"Slate",
+				"TraceAnalysis",
+				"TraceServices",
+				"TraceInsights",
+				"TraceInsightsCore",
+			});
+
+		}
+	}
+}
+
diff --git a/Engine/Plugins/IoStoreInsights/Source/IoStoreInsights/Private/Analyzers/IoStoreInsightsAnalyzer.cpp b/Engine/Plugins/IoStoreInsights/Source/IoStoreInsights/Private/Analyzers/IoStoreInsightsAnalyzer.cpp
new file mode 100644
--- /dev/null
+++ b/Engine/Plugins/IoStoreInsights/Source/IoStoreInsights/Private/Analyzers/IoStoreInsightsAnalyzer.cpp
@@ -0,0 +1,254 @@
+// Copyright Epic Games, Inc. All Rights Reserved.
+
+#include "Analyzers/IoStoreInsightsAnalyzer.h"
+#include "Model/IoStoreInsightsProvider.h"
+#include "TraceServices/Model/MetadataProvider.h"
+#include "TraceServices/Model/Definitions.h"
+#include "TraceServices/Model/Strings.h"
+#include "TraceServices/Model/AnalysisSession.h"
+#include "Common/ProviderLock.h"
+#include "HAL/LowLevelMemTracker.h"
+
+namespace UE::IoStoreInsights
+{
+	FIoStoreInsightsAnalyzer::FIoStoreInsightsAnalyzer(TraceServices::IAnalysisSession& InSession, FIoStoreInsightsProvider& InIoStoreProvider)
+		: Session(InSession)
+		, Provider(InIoStoreProvider)
+	{
+	}
+
+
+
+	void FIoStoreInsightsAnalyzer::OnAnalysisBegin(const FOnAnalysisContext& Context)
+	{
+		Trace::IAnalyzer::FInterfaceBuilder& Builder = Context.InterfaceBuilder;
+
+		Builder.RouteEvent(RouteId_BackendName,      "IoStore",  "BackendName");
+		Builder.RouteEvent(RouteId_RequestCreate,    "IoStore",  "RequestCreate");
+		Builder.RouteEvent(RouteId_RequestStarted,   "IoStore",  "RequestStarted");
+		Builder.RouteEvent(RouteId_RequestCompleted, "IoStore",  "RequestCompleted");
+		Builder.RouteEvent(RouteId_RequestFailed,    "IoStore",  "RequestFailed");
+		Builder.RouteEvent(RouteId_PackageMapping,   "Package",  "PackageMapping");
+	}
+
+
+
+	bool FIoStoreInsightsAnalyzer::OnEvent(uint16 RouteId, EStyle Style, const FOnEventContext& Context)
+	{
+		LLM_SCOPE_BYNAME(TEXT("Insights/FIoStoreInsightsAnalyzer"));
+		const auto& EventData = Context.EventData;
+
+		TraceServices::FAnalysisSessionEditScope _(Session);
+		switch (RouteId)
+		{
+			case RouteId_BackendName:
+			{
+				const uint64 Backend = EventData.GetValue<uint64>("BackendHandle");
+				FString BackendName;
+				if (EventData.GetString("Name", BackendName))
+				{
+					Provider.AddBackendName(Backend, Session.StoreString(BackendName));
+				}
+			}
+			break;
+
+			case RouteId_RequestCreate:
+			{
+				const uint64 Cycle = EventData.GetValue<uint64>("Cycle");
+				double Time = Context.EventTime.AsSeconds(Cycle);
+				Session.UpdateDurationSeconds(Time);
+
+				const uint64 RequestHandle = EventData.GetValue<uint64>("RequestHandle");
+				const uint64 BatchHandle   = EventData.GetValue<uint64>("BatchHandle");
+				const uint32 CallstackId   = EventData.GetValue<uint32>("CallstackId");
+				const uint64 Offset        = EventData.GetValue<uint64>("Offset");
+				const uint64 Size          = EventData.GetValue<uint64>("Size");
+				const uint32 ChunkIdHash   = EventData.GetValue<uint32>("ChunkIdHash");
+				const uint8 ChunkType      = EventData.GetValue<uint8>("ChunkType");
+				const uint32 ThreadId      = Context.ThreadInfo.GetId();
+				check( !ActiveRequestsMap.Contains(RequestHandle) );
+
+				FString PackageName, ExtraTag;
+				uint64 PackageId;
+				GetPackageDetailFromMetadata(Context, ThreadId, PackageName, ExtraTag, PackageId);
+				uint32 IoStoreRequestIndex = Provider.GetIoStoreRequestIndex(ChunkIdHash, ChunkType, Offset, Size, CallstackId, PackageId, *PackageName, *ExtraTag);
+
+				uint64 ActivityIndex = Provider.BeginIoStoreActivity(IoStoreRequestIndex, EIoStoreActivityType::Request_Pending, ThreadId, 0, Time);
+				ActiveRequestsMap.Emplace( RequestHandle, {IoStoreRequestIndex, ActivityIndex} );
+			}
+			break;
+
+			case RouteId_RequestStarted:
+			{
+				const uint64 Cycle = EventData.GetValue<uint64>("Cycle");
+				double Time = Context.EventTime.AsSeconds(Cycle);
+				Session.UpdateDurationSeconds(Time);
+
+				const uint64 RequestHandle = EventData.GetValue<uint64>("RequestHandle");
+				const uint64 BackendHandle = EventData.GetValue<uint64>("BackendHandle");
+				const uint32 ThreadId      = Context.ThreadInfo.GetId();
+			
+				const FPendingRequest* Request = ActiveRequestsMap.Find(RequestHandle);
+				if (Request == nullptr)
+				{
+					//IOSTORETRACE_WARNING("RequestStarted: request was not created!");
+					Request = &ActiveRequestsMap.Emplace( RequestHandle, {Provider.GetUnknownIoStoreRequestIndex(), 0} );
+				}
+				else
+				{
+					Provider.EndIoStoreActivity(Request->IoStoreRequestIndex, Request->CreateActivityIndex, 0, false, Time);
+				}
+
+				uint64 ReadIndex = Provider.BeginIoStoreActivity(Request->IoStoreRequestIndex, EIoStoreActivityType::Request_Read, ThreadId, BackendHandle, Time);
+				FPendingActivity& Read = ActiveReadsMap.Add(RequestHandle);
+				Read.IoStoreRequestIndex = Request->IoStoreRequestIndex;
+				Read.ActivityIndex = ReadIndex;
+			}
+			break;
+
+			case RouteId_RequestCompleted:
+			case RouteId_RequestFailed:
+			{
+				const uint64 Cycle = EventData.GetValue<uint64>("Cycle");
+				double Time = Context.EventTime.AsSeconds(Cycle);
+				Session.UpdateDurationSeconds(Time);
+
+				const bool bFailed         = (RouteId == RouteId_RequestFailed);
+				const uint64 RequestHandle = EventData.GetValue<uint64>("RequestHandle");
+				const uint64 Size          = bFailed ? 0 : EventData.GetValue<uint64>("Size");
+
+				const FPendingActivity* FindRead = ActiveReadsMap.Find(RequestHandle);
+				if (FindRead)
+				{
+					Provider.EndIoStoreActivity(FindRead->IoStoreRequestIndex, FindRead->ActivityIndex, Size, bFailed, Time);
+					ActiveReadsMap.Remove(RequestHandle);
+					ActiveRequestsMap.Remove(RequestHandle); // note: we don't have a separate "request destroy", so clean everything up once it has been read
+				}
+				else
+				{
+					//IOSTORETRACE_WARNING("RequestCompleted: RequestStarted not traced!");
+				}
+			}
+			break;
+
+			case RouteId_PackageMapping:
+			{
+				const TraceServices::IDefinitionProvider* DefinitionProvider = TraceServices::ReadDefinitionProvider(Session);
+				if (DefinitionProvider)
+				{
+					const uint64 PackageId    = EventData.GetValue<uint64>("Id");
+					const auto PackageNameRef = EventData.GetReferenceValue<uint32>("Package");
+					const auto PackageName    = DefinitionProvider->Get<TraceServices::FStringDefinition>(PackageNameRef);
+					if (PackageName && *PackageName->Display != '\0')
+					{
+						Provider.AddPackageMapping(PackageId, PackageName->Display);
+					}
+				}
+			}
+			break;
+
+		}
+		return true;
+	}
+
+
+
+	bool FIoStoreInsightsAnalyzer::GetPackageDetailFromMetadata( const FOnEventContext& Context, uint32 ThreadId, FString& OutPackageName, FString& OutExtraTag, uint64& OutPackageId ) const
+	{
+		const TraceServices::IMetadataProvider* MetadataProvider = ReadMetadataProvider(Session);
+
+		// lookup metadata id
+		uint32 MetadataId = TraceServices::IMetadataProvider::InvalidMetadataId;
+		{
+			TraceServices::IEditableMetadataProvider* EditableMetadataProvider = EditMetadataProvider(Session);
+			TraceServices::FProviderEditScopeLock MetadataProviderEditLock(*EditableMetadataProvider);
+
+			MetadataId = EditableMetadataProvider->PinAndGetId(ThreadId);
+		}
+		if (MetadataId == TraceServices::IMetadataProvider::InvalidMetadataId)
+		{
+			return false;
+		}
+
+		// lookup metadata types (must be nicer way to do this!)
+		struct FMetadataType
+		{
+			uint16 MetadataType = 0;
+			const TraceServices::FMetadataSchema* Schema = nullptr;
+
+			void Init(const TraceServices::IMetadataProvider* MetadataProvider, int& PendingMetadata, FName Name)
+			{
+				MetadataType = MetadataProvider->GetRegisteredMetadataType(Name);
+				if (MetadataType != TraceServices::IMetadataProvider::InvalidMetadataId)
+				{
+					Schema = MetadataProvider->GetRegisteredMetadataSchema(MetadataType);
+					if (Schema)
+					{
+						PendingMetadata++;
+					}
+				}
+			}
+		} MetadataTypes[3];
+		int PendingMetadata = 0;
+		{
+			TraceServices::FProviderReadScopeLock MetadataProviderReadLock(*MetadataProvider);
+			MetadataTypes[0].Init(MetadataProvider, PendingMetadata, TEXT("Asset"));
+			MetadataTypes[1].Init(MetadataProvider, PendingMetadata, TEXT("PackageId"));
+			MetadataTypes[2].Init(MetadataProvider, PendingMetadata, TEXT("IoStoreTag"));
+		}
+
+		// find the asset metadata associated with the current thread callstack
+		const TraceServices::IDefinitionProvider* DefinitionProvider = TraceServices::ReadDefinitionProvider(Session);
+		if (MetadataId != TraceServices::IMetadataProvider::InvalidMetadataId && DefinitionProvider && PendingMetadata > 0)
+		{
+			TraceServices::FProviderReadScopeLock MetadataProviderReadLock(*MetadataProvider);
+			MetadataProvider->EnumerateMetadata(ThreadId, MetadataId,
+				[&MetadataTypes, &PendingMetadata, DefinitionProvider, &OutPackageName, &OutPackageId, &OutExtraTag](uint32 StackDepth, uint16 Type, const void* Data, uint32 Size) -> bool
+				{
+					if (Type == MetadataTypes[0].MetadataType && MetadataTypes[0].Schema) // Asset metadata
+					{
+						TraceServices::FProviderReadScopeLock DefinitionProviderReadLock(*DefinitionProvider);
+						const auto Reader = MetadataTypes[0].Schema->Reader();
+						const auto PackageNameRef = Reader.GetValueAs<UE::Trace::FEventRef32>((uint8*)Data, 2); // { 0:Name, 1:Class, 2:Package }
+						const auto PackageName = DefinitionProvider->Get<TraceServices::FStringDefinition>(*PackageNameRef);
+						if (PackageName)
+						{
+							OutPackageName = PackageName->Display;
+						}
+						PendingMetadata--;
+					}
+					else if (Type == MetadataTypes[1].MetadataType && MetadataTypes[1].Schema) // PackageId metadata
+					{
+						TraceServices::FProviderReadScopeLock DefinitionProviderReadLock(*DefinitionProvider);
+						const auto Reader = MetadataTypes[1].Schema->Reader();
+						const auto PackageId = Reader.GetValueAs<uint64>((uint8*)Data, 0); // { 0:PackageId }
+						if (PackageId)
+						{
+							OutPackageId = *PackageId;
+						}
+						PendingMetadata--;
+					}
+					else if (Type == MetadataTypes[2].MetadataType && MetadataTypes[2].Schema) // IoStoreTag metadata
+					{
+						TraceServices::FProviderReadScopeLock DefinitionProviderReadLock(*DefinitionProvider);
+						const auto Reader = MetadataTypes[2].Schema->Reader();
+						const auto TagNameRef = Reader.GetValueAs<UE::Trace::FEventRef32>((uint8*)Data, 0); // { 0:Tag }
+						const auto TagName = DefinitionProvider->Get<TraceServices::FStringDefinition>(*TagNameRef);
+						if (TagName)
+						{
+							OutExtraTag = TagName->Display;
+						}
+
+						PendingMetadata--;
+					}
+					return (PendingMetadata > 0);
+				});
+
+			return true;
+		}
+
+		return false;
+	}
+
+} //namespace UE::IoStoreInsights
+
diff --git a/Engine/Plugins/IoStoreInsights/Source/IoStoreInsights/Private/Analyzers/IoStoreInsightsAnalyzer.h b/Engine/Plugins/IoStoreInsights/Source/IoStoreInsights/Private/Analyzers/IoStoreInsightsAnalyzer.h
new file mode 100644
--- /dev/null
+++ b/Engine/Plugins/IoStoreInsights/Source/IoStoreInsights/Private/Analyzers/IoStoreInsightsAnalyzer.h
@@ -0,0 +1,57 @@
+// Copyright Epic Games, Inc. All Rights Reserved.
+
+#pragma once
+
+#include "Trace/Analyzer.h"
+#include "Containers/Map.h"
+
+namespace TraceServices
+{ 
+	class IAnalysisSession; 
+}
+
+namespace UE::IoStoreInsights
+{ 
+	class FIoStoreInsightsProvider;
+
+	class FIoStoreInsightsAnalyzer : public UE::Trace::IAnalyzer
+	{
+	public:
+		FIoStoreInsightsAnalyzer(TraceServices::IAnalysisSession& InSession, FIoStoreInsightsProvider& InIoStoreProvider);
+
+		virtual void OnAnalysisBegin(const FOnAnalysisContext& Context) override;
+		virtual bool OnEvent(uint16 RouteId, EStyle Style, const FOnEventContext& Context) override;
+
+	private:
+		bool GetPackageDetailFromMetadata( const FOnEventContext& Context, uint32 ThreadId, FString& OutPackageName, FString& OutExtraTag, uint64& OutPackageId ) const;
+
+		enum : uint16
+		{
+			RouteId_BackendName,
+			RouteId_RequestCreate,
+			RouteId_RequestStarted,
+			RouteId_RequestCompleted,
+			RouteId_RequestFailed,
+			RouteId_PackageMapping,
+		};
+
+		struct FPendingActivity
+		{
+			uint64 ActivityIndex;
+			uint32 IoStoreRequestIndex;
+		};
+
+		struct FPendingRequest
+		{
+			uint32 IoStoreRequestIndex;
+			uint64 CreateActivityIndex;
+		};
+
+		TraceServices::IAnalysisSession& Session;
+		FIoStoreInsightsProvider& Provider;
+		TMap<uint64, FPendingActivity> ActiveReadsMap;
+		TMap<uint64, FPendingRequest> ActiveRequestsMap;
+		TMap<uint64, uint64> ActiveBatchMap;
+	};
+
+} // namespace UE::IoStoreInsights
diff --git a/Engine/Plugins/IoStoreInsights/Source/IoStoreInsights/Private/IoStoreInsightsModule.cpp b/Engine/Plugins/IoStoreInsights/Source/IoStoreInsights/Private/IoStoreInsightsModule.cpp
new file mode 100644
--- /dev/null
+++ b/Engine/Plugins/IoStoreInsights/Source/IoStoreInsights/Private/IoStoreInsightsModule.cpp
@@ -0,0 +1,24 @@
+// Copyright Epic Games, Inc. All Rights Reserved.
+
+#include "IoStoreInsightsModule.h"
+#include "Modules/ModuleManager.h"
+#include "Features/IModularFeatures.h"
+
+namespace UE::IoStoreInsights
+{
+	void FIoStoreInsightsModule::StartupModule()
+	{
+		IModularFeatures::Get().RegisterModularFeature(TraceServices::ModuleFeatureName, &TraceModule);
+		IModularFeatures::Get().RegisterModularFeature(Insights::Timing::TimingViewExtenderFeatureName, &TimingViewExtender);
+	}
+
+	void FIoStoreInsightsModule::ShutdownModule()
+	{
+		IModularFeatures::Get().UnregisterModularFeature(TraceServices::ModuleFeatureName, &TraceModule);
+		IModularFeatures::Get().UnregisterModularFeature(Insights::Timing::TimingViewExtenderFeatureName, &TimingViewExtender);
+	}
+
+
+} //namespace UE::IoStoreInsights
+
+IMPLEMENT_MODULE(UE::IoStoreInsights::FIoStoreInsightsModule, IoStoreInsights);
diff --git a/Engine/Plugins/IoStoreInsights/Source/IoStoreInsights/Private/IoStoreInsightsModule.h b/Engine/Plugins/IoStoreInsights/Source/IoStoreInsights/Private/IoStoreInsightsModule.h
new file mode 100644
--- /dev/null
+++ b/Engine/Plugins/IoStoreInsights/Source/IoStoreInsights/Private/IoStoreInsightsModule.h
@@ -0,0 +1,22 @@
+// Copyright Epic Games, Inc. All Rights Reserved.
+
+#pragma once
+
+#include "Modules/ModuleInterface.h"
+#include "IoStoreInsightsTraceModule.h"
+#include "ViewModels/IoStoreInsightsTimingViewExtender.h"
+
+namespace UE::IoStoreInsights
+{
+	class FIoStoreInsightsModule : public IModuleInterface
+	{
+	public:
+		virtual void StartupModule() override;
+		virtual void ShutdownModule() override;
+
+	private:
+		FIoStoreInsightsTraceModule TraceModule;
+		FIoStoreInsightsTimingViewExtender TimingViewExtender;
+	};
+
+} // namespace UE::IoStoreInsights
diff --git a/Engine/Plugins/IoStoreInsights/Source/IoStoreInsights/Private/IoStoreInsightsTraceModule.cpp b/Engine/Plugins/IoStoreInsights/Source/IoStoreInsights/Private/IoStoreInsightsTraceModule.cpp
new file mode 100644
--- /dev/null
+++ b/Engine/Plugins/IoStoreInsights/Source/IoStoreInsights/Private/IoStoreInsightsTraceModule.cpp
@@ -0,0 +1,43 @@
+// Copyright Epic Games, Inc. All Rights Reserved.
+
+#include "IoStoreInsightsTraceModule.h"
+#include "Model/IoStoreInsightsProvider.h"
+#include "Analyzers/IoStoreInsightsAnalyzer.h"
+#include "TraceServices/Model/AnalysisSession.h"
+
+namespace UE::IoStoreInsights
+{
+	FName FIoStoreInsightsTraceModule::ModuleName("TraceModule_IoStore");
+
+
+
+	void FIoStoreInsightsTraceModule::GetModuleInfo(TraceServices::FModuleInfo& OutModuleInfo)
+	{
+		OutModuleInfo.Name = ModuleName;
+		OutModuleInfo.DisplayName = TEXT("IoStore");
+	}
+
+
+
+	void FIoStoreInsightsTraceModule::OnAnalysisBegin(TraceServices::IAnalysisSession& InSession)
+	{
+		TSharedPtr<FIoStoreInsightsProvider> IoStoreProvider = MakeShared<FIoStoreInsightsProvider>(InSession);
+		InSession.AddProvider(FIoStoreInsightsProvider::ProviderName, IoStoreProvider);
+		InSession.AddAnalyzer(new FIoStoreInsightsAnalyzer(InSession, *IoStoreProvider));
+	}
+
+
+
+	void FIoStoreInsightsTraceModule::GetLoggers(TArray<const TCHAR *>& OutLoggers)
+	{
+		OutLoggers.Add(TEXT("IoStore"));
+	}
+
+
+
+	void FIoStoreInsightsTraceModule::GenerateReports(const TraceServices::IAnalysisSession& Session, const TCHAR* CmdLine, const TCHAR* OutputDirectory)
+	{
+	}
+
+} //namespace UE::IoStoreInsights
+
diff --git a/Engine/Plugins/IoStoreInsights/Source/IoStoreInsights/Private/IoStoreInsightsTraceModule.h b/Engine/Plugins/IoStoreInsights/Source/IoStoreInsights/Private/IoStoreInsightsTraceModule.h
new file mode 100644
--- /dev/null
+++ b/Engine/Plugins/IoStoreInsights/Source/IoStoreInsights/Private/IoStoreInsightsTraceModule.h
@@ -0,0 +1,26 @@
+// Copyright Epic Games, Inc. All Rights Reserved.
+
+#pragma once
+
+#include "TraceServices/ModuleService.h"
+#include "ViewModels/IoStoreInsightsTimingViewExtender.h"
+
+namespace UE::IoStoreInsights
+{
+	class FIoStoreInsightsTraceModule : public TraceServices::IModule
+	{
+	public:
+		virtual ~FIoStoreInsightsTraceModule() = default;
+
+		// TraceServices::IModule interface
+		virtual void GetModuleInfo(TraceServices::FModuleInfo& OutModuleInfo) override;
+		virtual void OnAnalysisBegin(TraceServices::IAnalysisSession& Session) override;
+		virtual void GetLoggers(TArray<const TCHAR *>& OutLoggers) override;
+		virtual void GenerateReports(const TraceServices::IAnalysisSession& Session, const TCHAR* CmdLine, const TCHAR* OutputDirectory) override;
+		virtual const TCHAR* GetCommandLineArgument() override { return TEXT("iostoretrace"); }
+
+	private:
+		static FName ModuleName;
+	};
+
+} // namespace UE::IoStoreInsights
diff --git a/Engine/Plugins/IoStoreInsights/Source/IoStoreInsights/Private/Model/IoStoreInsightsProvider.cpp b/Engine/Plugins/IoStoreInsights/Source/IoStoreInsights/Private/Model/IoStoreInsightsProvider.cpp
new file mode 100644
--- /dev/null
+++ b/Engine/Plugins/IoStoreInsights/Source/IoStoreInsights/Private/Model/IoStoreInsightsProvider.cpp
@@ -0,0 +1,152 @@
+// Copyright Epic Games, Inc. All Rights Reserved.
+
+#include "IoStoreInsightsProvider.h"
+
+namespace UE::IoStoreInsights
+{
+	FName FIoStoreInsightsProvider::ProviderName("IoStoreProvider");
+
+
+
+	FIoStoreInsightsProvider::FIoStoreInsightsProvider(TraceServices::IAnalysisSession& InSession)
+		: Session(InSession)
+		, IoStoreRequests(InSession.GetLinearAllocator(), 1024)
+		, IoStoreRequestStates(InSession.GetLinearAllocator(), 1024)
+	{
+	}
+
+
+
+	void FIoStoreInsightsProvider::EnumerateIoStoreRequests(TFunctionRef<bool(const FIoStoreRequest&, const Timeline&)> Callback) const
+	{
+		for (uint64 IoStoreRequestIndex = 0, IoStoreRequestCount = IoStoreRequests.Num(); IoStoreRequestIndex < IoStoreRequestCount; ++IoStoreRequestIndex)
+		{
+			const FIoStoreRequestInfoInternal& IoStoreRequestInfoInternal = IoStoreRequests[IoStoreRequestIndex];
+			if (!Callback(IoStoreRequestInfoInternal.IoStoreRequestInfo, *IoStoreRequestInfoInternal.ActivityTimeline))
+			{
+				return;
+			}
+		}
+	}
+
+
+
+	uint32 FIoStoreInsightsProvider::GetIoStoreRequestIndex(uint32 ChunkIdHash, uint8 ChunkType, uint64 Offset, uint64 Size, uint32 CallstackId, uint64 PackageId, const TCHAR* PackageName, const TCHAR* ExtraTag)
+	{
+		uint32 IoStoreRequestIndex = static_cast<uint32>(IoStoreRequests.Num());
+		FIoStoreRequestInfoInternal& IoStoreRequestInfo = IoStoreRequests.PushBack();
+		IoStoreRequestInfo.IoStoreRequestInfo.IoStoreRequestIndex = IoStoreRequestIndex;
+		IoStoreRequestInfo.IoStoreRequestInfo.ChunkIdHash = ChunkIdHash;
+		IoStoreRequestInfo.IoStoreRequestInfo.ChunkType = ChunkType;
+		IoStoreRequestInfo.IoStoreRequestInfo.Offset = Offset;
+		IoStoreRequestInfo.IoStoreRequestInfo.Size = Size;
+		IoStoreRequestInfo.IoStoreRequestInfo.PackageId = PackageId;
+		IoStoreRequestInfo.IoStoreRequestInfo.CallstackId = CallstackId;
+		IoStoreRequestInfo.IoStoreRequestInfo.PackageName = Session.StoreString(PackageName);
+		IoStoreRequestInfo.IoStoreRequestInfo.ExtraTag = Session.StoreString(ExtraTag);
+		IoStoreRequestInfo.ActivityTimeline = MakeShared<TimelineInternal>(Session.GetLinearAllocator());
+		return IoStoreRequestIndex;
+	}
+
+
+
+	uint32 FIoStoreInsightsProvider::GetUnknownIoStoreRequestIndex()
+	{
+		return GetIoStoreRequestIndex(0, 0, 0, 0, 0, 0, TEXT("unknown"),TEXT(""));
+	}
+
+
+
+	uint64 FIoStoreInsightsProvider::BeginIoStoreActivity(uint32 IoStoreRequestIndex, EIoStoreActivityType Type, uint32 ThreadId, uint64 BackendHandle, double Time)
+	{
+		FIoStoreRequestInfoInternal& IoStoreRequestInfo = IoStoreRequests[IoStoreRequestIndex];
+		FIoStoreActivity& IoStoreActivity = IoStoreRequestStates.PushBack();
+		IoStoreActivity.IoStoreRequest = &IoStoreRequestInfo.IoStoreRequestInfo;
+		IoStoreActivity.BackendHandle = BackendHandle;
+		IoStoreActivity.ActualSize = 0;
+		IoStoreActivity.StartTime = 0;
+		IoStoreActivity.EndTime = std::numeric_limits<double>::infinity();
+		IoStoreActivity.ThreadId = ThreadId;
+		IoStoreActivity.ActivityType = Type;
+		IoStoreActivity.Failed = false;
+	
+		return IoStoreRequestInfo.ActivityTimeline->AppendBeginEvent(Time, &IoStoreActivity);
+	}
+
+
+
+	void FIoStoreInsightsProvider::EndIoStoreActivity(uint32 IoStoreRequestIndex, uint64 ActivityIndex, uint64 ActualSize, bool Failed, double Time)
+	{
+		FIoStoreRequestInfoInternal& IoStoreRequestInfo = IoStoreRequests[IoStoreRequestIndex];
+		FIoStoreActivity* Activity = IoStoreRequestInfo.ActivityTimeline->EndEvent(ActivityIndex, Time);
+		check(Activity->IoStoreRequest == &IoStoreRequestInfo.IoStoreRequestInfo);
+		Activity->ActualSize = ActualSize;
+		Activity->Failed = Failed;
+		Activity->EndTime = Time;
+	}
+
+
+
+	void FIoStoreInsightsProvider::AddPackageMapping(uint64 PackageId, const TCHAR* PackageName)
+	{
+		PackageMap.Add(PackageId, PackageName);
+	}
+
+
+
+	const TCHAR* FIoStoreInsightsProvider::GetPackageName(uint64 PackageId) const
+	{
+		const TCHAR* const* PackageNamePtr = PackageMap.Find(PackageId);
+		if (PackageNamePtr)
+		{
+			return *PackageNamePtr;
+		}
+
+		return TEXT("");
+	}
+
+
+
+	void FIoStoreInsightsProvider::AddBackendName(uint64 BackendHandle, const TCHAR* BackendName)
+	{
+		BackendNameMap.Add(BackendHandle, BackendName);
+	}
+
+
+
+	const TCHAR* FIoStoreInsightsProvider::GetBackendName(uint64 BackendHandle) const
+	{
+		const TCHAR* const* BackendNamePtr = BackendNameMap.Find(BackendHandle);
+		return BackendNamePtr ? *BackendNamePtr : TEXT("unknown backend");
+	}
+
+
+
+	const FIoStoreRequest& FIoStoreInsightsProvider::GetIoStoreRequest(uint32 IoStoreRequestIndex) const
+	{
+		check(IoStoreRequestIndex < IoStoreRequests.Num());
+		return IoStoreRequests[IoStoreRequestIndex].IoStoreRequestInfo;
+	}
+
+
+
+	const IIoStoreInsightsProvider* ReadIoStoreInsightsProvider(const TraceServices::IAnalysisSession& Session)
+	{
+		Session.ReadAccessCheck();
+		return Session.ReadProvider<IIoStoreInsightsProvider>(FIoStoreInsightsProvider::ProviderName);
+	}
+
+
+
+	const TCHAR* LexToString(EIoStoreActivityType ActivityType)
+	{
+		static const TCHAR* Items[(int)EIoStoreActivityType::Count] =
+		{
+			TEXT("Pending"),
+			TEXT("Read"),
+		};
+
+		return (ActivityType < EIoStoreActivityType::Count) ? Items[(int)ActivityType] : TEXT("Invalid");
+	}
+
+} // namespace UE::IoStoreInsights
\ No newline at end of file
diff --git a/Engine/Plugins/IoStoreInsights/Source/IoStoreInsights/Private/Model/IoStoreInsightsProvider.h b/Engine/Plugins/IoStoreInsights/Source/IoStoreInsights/Private/Model/IoStoreInsightsProvider.h
new file mode 100644
--- /dev/null
+++ b/Engine/Plugins/IoStoreInsights/Source/IoStoreInsights/Private/Model/IoStoreInsightsProvider.h
@@ -0,0 +1,62 @@
+// Copyright Epic Games, Inc. All Rights Reserved.
+
+#pragma once
+
+#include "TraceServices/Model/AnalysisSession.h"
+#include "IIoStoreInsightsProvider.h"
+
+namespace TraceServices 
+{ 
+	class IAnalysisSession; 
+}
+
+namespace UE::IoStoreInsights
+{
+	class FIoStoreInsightsProvider : public IIoStoreInsightsProvider
+	{
+	public:
+		static FName ProviderName;
+
+		struct FTimelineSettings
+		{
+			enum
+			{
+				EventsPerPage = 128
+			};
+		};
+		typedef TraceServices::TIntervalTimeline<FIoStoreActivity*, FTimelineSettings> TimelineInternal;
+
+		explicit FIoStoreInsightsProvider(TraceServices::IAnalysisSession& Session);
+		virtual ~FIoStoreInsightsProvider() {};
+
+		virtual void EnumerateIoStoreRequests(TFunctionRef<bool(const FIoStoreRequest&, const Timeline&)> Callback) const override;
+
+		uint32 GetIoStoreRequestIndex(uint32 ChunkIdHash, uint8 ChunkType, uint64 Offset, uint64 Size, uint32 CallstackId, uint64 PackageId, const TCHAR* PackageName, const TCHAR* ExtraTag);
+		uint32 GetUnknownIoStoreRequestIndex();
+
+		uint64 BeginIoStoreActivity(uint32 IoStoreRequestIndex, EIoStoreActivityType Type, uint32 ThreadId, uint64 BackendHandle, double Time);
+		void EndIoStoreActivity(uint32 IoStoreRequestIndex, uint64 ActivityIndex, uint64 ActualSize, bool Failed, double Time);
+
+		void AddPackageMapping(uint64 PackageId, const TCHAR* PackageName);
+		virtual const TCHAR* GetPackageName(uint64 PackageId) const override;
+
+		void AddBackendName(uint64 BackendHandle, const TCHAR* BackendName);
+		virtual const TCHAR* GetBackendName(uint64 BackendHandle) const override;
+
+		virtual const FIoStoreRequest& GetIoStoreRequest(uint32 IoStoreRequestIndex) const override;
+
+	private:
+		struct FIoStoreRequestInfoInternal
+		{
+			FIoStoreRequest IoStoreRequestInfo;
+			TSharedPtr<TimelineInternal> ActivityTimeline;
+		};
+
+		TraceServices::IAnalysisSession& Session;
+		TraceServices::TPagedArray<FIoStoreRequestInfoInternal> IoStoreRequests;
+		TraceServices::TPagedArray<FIoStoreActivity> IoStoreRequestStates;
+		TMap<uint64, const TCHAR*> PackageMap;
+		TMap<uint64, const TCHAR*> BackendNameMap;
+	};
+
+} // namespace UE::IoStoreInsights
diff --git a/Engine/Plugins/IoStoreInsights/Source/IoStoreInsights/Private/ViewModels/IoStoreInsightsTimingViewExtender.cpp b/Engine/Plugins/IoStoreInsights/Source/IoStoreInsights/Private/ViewModels/IoStoreInsightsTimingViewExtender.cpp
new file mode 100644
--- /dev/null
+++ b/Engine/Plugins/IoStoreInsights/Source/IoStoreInsights/Private/ViewModels/IoStoreInsightsTimingViewExtender.cpp
@@ -0,0 +1,334 @@
+// Copyright Epic Games, Inc. All Rights Reserved.
+
+#include "ViewModels/IoStoreInsightsTimingViewExtender.h"
+#include "ViewModels/IoStoreInsightsTrack.h"
+#include "Model/IoStoreInsightsProvider.h"
+#include "Framework/MultiBox/MultiBoxBuilder.h"
+#include "InsightsCore/Common/Stopwatch.h"
+#include "InsightsCore/Common/TimeUtils.h"
+#include "Insights/ITimingViewSession.h"
+#include "Insights/TimingProfilerCommon.h"
+
+#define LOCTEXT_NAMESPACE "IoStoreViewExtender"
+
+namespace UE::IoStoreInsights
+{
+	const uint32 FIoStoreInsightsViewSharedState::MaxLanes = 10000;
+
+
+
+	void FIoStoreInsightsViewSharedState::Tick(const TraceServices::IAnalysisSession& InAnalysisSession)
+	{
+		// cache session
+		if (AnalysisSession != &InAnalysisSession)
+		{
+			AnalysisSession = &InAnalysisSession;
+			bForceIoEventsUpdate = true;
+		}
+
+		// see if we should rebuild the shared state
+		if (bForceIoEventsUpdate)
+		{
+			bForceIoEventsUpdate = false;
+
+			Insights::FStopwatch Stopwatch;
+
+			// Enumerate all IO events and cache them.
+			Stopwatch.Start();
+			AllIoEvents.Reset();
+			TArray<TSharedPtr<FIoStoreRequestState>> RequestStates;
+			{
+				TraceServices::FAnalysisSessionReadScope SessionReadScope(InAnalysisSession);
+				const IIoStoreInsightsProvider* IoStoreActivityProvider = InAnalysisSession.ReadProvider<IIoStoreInsightsProvider>(FIoStoreInsightsProvider::ProviderName);
+
+				IoStoreActivityProvider->EnumerateIoStoreRequests([this, IoStoreActivityProvider, &RequestStates](const FIoStoreRequest& IoStoreRequest, const IIoStoreInsightsProvider::Timeline& Timeline)
+				{
+					TSharedPtr<FIoStoreRequestState> RequestState = RequestStates.Add_GetRef(MakeShared<FIoStoreRequestState>());
+					RequestState->IoStoreRequestIndex = IoStoreRequest.IoStoreRequestIndex;
+					RequestState->StartTime           = +std::numeric_limits<double>::infinity();
+					RequestState->EndTime             = -std::numeric_limits<double>::infinity();
+					RequestState->MaxConcurrentEvents = 0;
+					RequestState->StartingDepth       = 0;
+
+					TArray<double> ConcurrentEvents;
+					Timeline.EnumerateEvents(-std::numeric_limits<double>::infinity(), +std::numeric_limits<double>::infinity(),
+						[this, &RequestState, &IoStoreRequest, &Timeline, &ConcurrentEvents, IoStoreActivityProvider](double EventStartTime, double EventEndTime, uint32 EventDepth, const FIoStoreActivity* IoStoreActivity)
+						{
+							// view is easier to read with just read events, but some people may want to see how long a request has been waiting for
+							if (bShowOnlyReadEvents && IoStoreActivity->ActivityType != EIoStoreActivityType::Request_Read)
+							{
+								return TraceServices::EEventEnumerate::Continue;
+							}
+
+							// events should be ordered by start time, but RequestState->StartTime may not be initialized
+							ensure(RequestState->StartTime == +std::numeric_limits<double>::infinity() || EventStartTime >= RequestState->StartTime);
+							if (EventStartTime < RequestState->StartTime)
+							{
+								RequestState->StartTime = EventStartTime;
+							}
+							if (EventEndTime > RequestState->EndTime)
+							{
+								RequestState->EndTime = EventEndTime;
+							}
+
+							uint32 LocalDepth = MAX_uint32;
+							for (int32 i = 0; i < ConcurrentEvents.Num(); ++i)
+							{
+								if (EventStartTime >= ConcurrentEvents[i])
+								{
+									LocalDepth = i;
+									ConcurrentEvents[i] = EventEndTime;
+									break;
+								}
+							}
+							if (LocalDepth == MAX_uint32)
+							{
+								LocalDepth = ConcurrentEvents.Num();
+								ConcurrentEvents.Add(EventEndTime);
+								RequestState->MaxConcurrentEvents = ConcurrentEvents.Num();
+							}
+
+
+							// slightly hacky: This becomes FIoStoreTimingEvent::Type so you can double-click on an event and see every IoStore event related to the same pacakge/bulk data (via FIoStoreInsightsTrack::BuildFilteredDrawState)
+							uint64 TimingEventType = HashCombine( PointerHash(IoStoreRequest.PackageName), PointerHash(IoStoreRequest.ExtraTag));
+							TimingEventType = HashCombine(TimingEventType, IoStoreRequest.PackageId);
+
+							uint32 Type = ((uint32)IoStoreActivity->ActivityType & 0x0F) | (IoStoreActivity->Failed ? 0x80 : 0);
+							AllIoEvents.Add(FIoStoreEventState{ RequestState, EventStartTime, EventEndTime, IoStoreActivity->ActualSize, IoStoreActivity->BackendHandle, TimingEventType, LocalDepth, Type });
+							return TraceServices::EEventEnumerate::Continue;
+						});
+
+					return true;
+				});
+			}
+			Stopwatch.Stop();
+			UE_LOG(LogTimingProfiler, Log, TEXT("[IO] Enumerated %s IoStore activities in %s."), *FText::AsNumber(RequestStates.Num()).ToString(), *Insights::FormatTimeAuto(Stopwatch.GetAccumulatedTime()));
+
+			// Sort cached items by Start Time
+			Stopwatch.Restart();			
+			RequestStates.Sort([](const TSharedPtr<FIoStoreRequestState>& A, const TSharedPtr<FIoStoreRequestState>& B) { return A->StartTime < B->StartTime; });
+			AllIoEvents.Sort([](const FIoStoreEventState& A, const FIoStoreEventState& B) { return A.StartTime < B.StartTime; });
+			Stopwatch.Stop();
+			UE_LOG(LogTimingProfiler, Log, TEXT("[IO] Sorted IoStore activities in %s."), *Insights::FormatTimeAuto(Stopwatch.GetAccumulatedTime()));
+
+			// Compute depth for IoStore activities (avoids overlaps).
+			if (RequestStates.Num() > 0)
+			{
+				struct FLane
+				{
+					double EndTime = 0.0f;
+				};
+				TArray<FLane> Lanes; // one lane per event depth, a IoStore activity occupies multiple lanes
+
+				Stopwatch.Restart();
+				for (const TSharedPtr<FIoStoreRequestState>& IoStoreActivityPtr : RequestStates)
+				{
+					FIoStoreRequestState& RequestState = *IoStoreActivityPtr;
+
+					// Find lane (avoiding overlaps with other IoStore activities).
+					int32 Depth = 0;
+					while (Depth < Lanes.Num())
+					{
+						bool bOverlap = false;
+						for (int32 LocalDepth = 0; LocalDepth < RequestState.MaxConcurrentEvents; ++LocalDepth)
+						{
+							if (Depth + LocalDepth >= Lanes.Num())
+							{
+								break;
+							}
+							const FLane& Lane = Lanes[Depth + LocalDepth];
+							if (RequestState.StartTime < Lane.EndTime)
+							{
+								bOverlap = true;
+								Depth += LocalDepth;
+								break;
+							}
+						}
+						if (!bOverlap)
+						{
+							break;
+						}
+						++Depth;
+					}
+
+					int32 NewLaneNum = Depth + RequestState.MaxConcurrentEvents;
+
+					if (NewLaneNum > MaxLanes)
+					{
+						// Snap to the bottom; allows overlaps in this case.
+						RequestState.StartingDepth = MaxLanes - RequestState.MaxConcurrentEvents;
+					}
+					else
+					{
+						if (NewLaneNum > Lanes.Num())
+						{
+							Lanes.AddDefaulted(NewLaneNum - Lanes.Num());
+						}
+
+						RequestState.StartingDepth = Depth;
+						for (int32 LocalDepth = 0; LocalDepth < RequestState.MaxConcurrentEvents; ++LocalDepth)
+						{
+							Lanes[Depth + LocalDepth].EndTime = RequestState.EndTime;
+						}
+					}
+				}
+				Stopwatch.Stop();
+				UE_LOG(LogTimingProfiler, Log, TEXT("[IO] Computed layout for IoStore activities in %s."), *Insights::FormatTimeAuto(Stopwatch.GetAccumulatedTime()));
+
+				Stopwatch.Restart();
+				for (FIoStoreEventState& Event : AllIoEvents)
+				{
+					Event.Depth += Event.RequestState->StartingDepth;
+					ensure(Event.Depth < MaxLanes);
+				}
+				Stopwatch.Stop();
+				UE_LOG(LogTimingProfiler, Log, TEXT("[IO] Updated depth for events in %s."), *Insights::FormatTimeAuto(Stopwatch.GetAccumulatedTime()));
+			}
+		}
+	}
+
+
+
+
+
+
+
+
+	void FIoStoreInsightsTimingViewExtender::OnBeginSession(Insights::Timing::ITimingViewSession& InSession)
+	{
+		FPerSessionData* PerSessionData = PerSessionDataMap.Find(&InSession);
+		if (PerSessionData == nullptr)
+		{
+			PerSessionData = &PerSessionDataMap.Add(&InSession);
+		}
+	}
+
+
+
+	void FIoStoreInsightsTimingViewExtender::OnEndSession(Insights::Timing::ITimingViewSession& InSession)
+	{
+		FPerSessionData* PerSessionData = PerSessionDataMap.Find(&InSession);
+		if (PerSessionData != nullptr)
+		{
+			if (PerSessionData->IoStoreActivityTrack.IsValid())
+			{
+				InSession.RemoveScrollableTrack(PerSessionData->IoStoreActivityTrack);
+				PerSessionData->IoStoreActivityTrack.Reset();
+			}
+
+			PerSessionDataMap.Remove(&InSession);
+		}
+	}
+
+
+
+	void FIoStoreInsightsTimingViewExtender::Tick(Insights::Timing::ITimingViewSession& InSession, const TraceServices::IAnalysisSession& InAnalysisSession)
+	{
+		bool bRequestUpdate = false;
+
+		// periodically refresh the shared data for each view while the session is still loading
+		{
+			TraceServices::FAnalysisSessionReadScope SessionReadScope(InAnalysisSession);
+			double SessionDuration = InAnalysisSession.GetDurationSeconds();
+			bool bAnalysisComplete = InAnalysisSession.IsAnalysisComplete();
+			if ((bAnalysisComplete && !bWasAnalysisComplete) || (SessionDuration - PreviousAnalysisSessionDuration) > 0.25f)
+			{
+				PreviousAnalysisSessionDuration = SessionDuration;
+				bWasAnalysisComplete = bAnalysisComplete;
+				bRequestUpdate = true;
+			}
+		}
+
+
+		FPerSessionData* PerSessionData = PerSessionDataMap.Find(&InSession);
+		if (PerSessionData != nullptr)
+		{
+			// create track if necessary
+			if (!PerSessionData->IoStoreActivityTrack.IsValid())
+			{
+				PerSessionData->IoStoreActivityTrack = MakeShared<FIoStoreInsightsTrack>(SharedState);
+				PerSessionData->IoStoreActivityTrack->SetOrder(FTimingTrackOrder::Last);
+				PerSessionData->IoStoreActivityTrack->SetVisibilityFlag(false);
+				InSession.AddScrollableTrack(PerSessionData->IoStoreActivityTrack);
+			}
+
+			// refresh the shared data if the track is dirty
+			bRequestUpdate |= PerSessionData->IoStoreActivityTrack->IsDirty() && PerSessionData->IoStoreActivityTrack->IsVisible();
+		}
+
+
+		// see if we should rebuild the shared state
+		if (bRequestUpdate)
+		{
+			SharedState.RequestUpdate();
+		}
+
+		SharedState.Tick(InAnalysisSession);
+	}
+
+
+
+	void FIoStoreInsightsTimingViewExtender::ExtendOtherTracksFilterMenu(Insights::Timing::ITimingViewSession& InSession, FMenuBuilder& InMenuBuilder)
+	{
+		FPerSessionData* PerSessionData = PerSessionDataMap.Find(&InSession);
+		if (PerSessionData != nullptr)
+		{
+			// note: can't use FUICommandList etc. because STimingView is private
+			InMenuBuilder.BeginSection("IoStore Activity", LOCTEXT("ContextMenu_Section_IoStoreActivity", "IoStore Activity"));
+			{
+				InMenuBuilder.AddMenuEntry(
+					LOCTEXT("ContextMenu_Item_IoStoreActivityTrack", "IoStore Activity Track"),
+					LOCTEXT("ContextMenu_Item_IoStoreActivityTrackTip", "Shows/hides the IoStore activity track"),
+					FSlateIcon(),
+					FUIAction(
+						FExecuteAction::CreateRaw(this, &FIoStoreInsightsTimingViewExtender::ToggleActivityTrackVisibility, PerSessionData->IoStoreActivityTrack),
+						FCanExecuteAction(),
+						FIsActionChecked::CreateRaw(this, &FIoStoreInsightsTimingViewExtender::IsIoStoreActivityTrackVisible, PerSessionData->IoStoreActivityTrack)
+					),
+					NAME_None,
+					EUserInterfaceActionType::ToggleButton);
+
+				InMenuBuilder.AddMenuEntry(
+					LOCTEXT("ContextMenu_Item_IoStoreShowAllEvents", "Only Reads (IoStore Activity Track)"),
+					LOCTEXT("ContextMenu_Item_IoStoreShowAllEventsTip", "Shows/hides all IoStore events"),
+					FSlateIcon(),
+					FUIAction(
+						FExecuteAction::CreateRaw(&SharedState, &FIoStoreInsightsViewSharedState::ToggleShowOnlyReadEvents),
+						FCanExecuteAction(),
+						FIsActionChecked::CreateRaw(&SharedState, &FIoStoreInsightsViewSharedState::IsShowingOnlyReadEvents)
+					),
+					NAME_None,
+					EUserInterfaceActionType::ToggleButton);
+
+			}
+			InMenuBuilder.EndSection();
+
+		}
+	}
+
+
+
+	bool FIoStoreInsightsTimingViewExtender::IsIoStoreActivityTrackVisible(TSharedPtr<FIoStoreInsightsTrack> IoStoreActivityTrack) const
+	{
+		return IoStoreActivityTrack && IoStoreActivityTrack->IsVisible();
+	}
+
+
+
+	void FIoStoreInsightsTimingViewExtender::ToggleActivityTrackVisibility(TSharedPtr<FIoStoreInsightsTrack> IoStoreActivityTrack)
+	{
+		if (IoStoreActivityTrack.IsValid())
+		{
+			IoStoreActivityTrack->ToggleVisibility();
+		}
+
+		const bool bIsActivityTrackVisible = IsIoStoreActivityTrackVisible(IoStoreActivityTrack);
+		if (bIsActivityTrackVisible)
+		{
+			SharedState.RequestUpdate();
+		}
+	}
+
+} // namespace UE::IoStoreInsights
+
+#undef LOCTEXT_NAMESPACE
diff --git a/Engine/Plugins/IoStoreInsights/Source/IoStoreInsights/Private/ViewModels/IoStoreInsightsTimingViewExtender.h b/Engine/Plugins/IoStoreInsights/Source/IoStoreInsights/Private/ViewModels/IoStoreInsightsTimingViewExtender.h
new file mode 100644
--- /dev/null
+++ b/Engine/Plugins/IoStoreInsights/Source/IoStoreInsights/Private/ViewModels/IoStoreInsightsTimingViewExtender.h
@@ -0,0 +1,92 @@
+// Copyright Epic Games, Inc. All Rights Reserved.
+
+#pragma once
+
+#include "CoreTypes.h"
+#include "Templates/SharedPointer.h"
+#include "Insights/ITimingViewExtender.h"
+
+namespace TraceServices
+{ 
+	class IAnalysisSession; 
+}
+
+namespace UE::Insights::Timing 
+{ 
+	class ITimingViewSession; 
+}
+
+namespace UE::IoStoreInsights
+{
+	class FIoStoreInsightsTrack;
+
+	struct FIoStoreRequestState
+	{
+		uint32 IoStoreRequestIndex;
+		double StartTime;
+		double EndTime;
+		int32 MaxConcurrentEvents; // only used during layout
+		uint32 StartingDepth; // only used during layout
+	};
+
+	struct FIoStoreEventState
+	{
+		TSharedPtr<FIoStoreRequestState> RequestState;
+		double StartTime;
+		double EndTime;
+		uint64 ActualSize;
+		uint64 BackendHandle;
+		uint64 TimingEventType;
+		uint32 Depth; // During update, this is local within a track - then it's set to a global depth
+		uint32 Type; // EIoStoreActivityType + "Failed" flag
+	};
+
+	class FIoStoreInsightsViewSharedState
+	{
+	public:
+		void Tick(const TraceServices::IAnalysisSession& InAnalysisSession);
+
+		const TraceServices::IAnalysisSession& GetAnalysisSession() const { return *AnalysisSession; }
+		const TArray<FIoStoreEventState>& GetAllEvents() const { return AllIoEvents; }
+
+		void RequestUpdate() { bForceIoEventsUpdate = true; }
+
+		bool IsShowingOnlyReadEvents() const { return bShowOnlyReadEvents; }
+		void ToggleShowOnlyReadEvents() { bShowOnlyReadEvents = !bShowOnlyReadEvents; RequestUpdate(); }
+
+		static const uint32 MaxLanes;
+
+	private:
+		const TraceServices::IAnalysisSession* AnalysisSession = nullptr; // cached in Tick()
+		bool bForceIoEventsUpdate = false;
+		bool bShowOnlyReadEvents = true;
+		TArray<FIoStoreEventState> AllIoEvents;
+	};
+
+	class FIoStoreInsightsTimingViewExtender : public Insights::Timing::ITimingViewExtender
+	{
+	public:
+		virtual void OnBeginSession(Insights::Timing::ITimingViewSession& InSession) override;
+		virtual void OnEndSession(Insights::Timing::ITimingViewSession& InSession) override;
+		virtual void Tick(Insights::Timing::ITimingViewSession& InSession, const TraceServices::IAnalysisSession& InAnalysisSession) override;
+		virtual void ExtendOtherTracksFilterMenu(Insights::Timing::ITimingViewSession& InSession, FMenuBuilder& InMenuBuilder) override;
+
+	private:
+		bool IsIoStoreActivityTrackVisible(TSharedPtr<FIoStoreInsightsTrack> IoStoreActivityTrack) const;
+		void ToggleActivityTrackVisibility(TSharedPtr<FIoStoreInsightsTrack> IoStoreActivityTrack);
+
+
+		bool bWasAnalysisComplete = false;
+		double PreviousAnalysisSessionDuration = 0.0f;
+		FIoStoreInsightsViewSharedState SharedState;
+
+		struct FPerSessionData
+		{
+			TSharedPtr<FIoStoreInsightsTrack> IoStoreActivityTrack;
+
+		};
+		TMap<Insights::Timing::ITimingViewSession*, FPerSessionData> PerSessionDataMap;
+
+
+	};
+} // namespace UE::IoStoreInsights
diff --git a/Engine/Plugins/IoStoreInsights/Source/IoStoreInsights/Private/ViewModels/IoStoreInsightsTrack.cpp b/Engine/Plugins/IoStoreInsights/Source/IoStoreInsights/Private/ViewModels/IoStoreInsightsTrack.cpp
new file mode 100644
--- /dev/null
+++ b/Engine/Plugins/IoStoreInsights/Source/IoStoreInsights/Private/ViewModels/IoStoreInsightsTrack.cpp
@@ -0,0 +1,302 @@
+// Copyright Epic Games, Inc. All Rights Reserved.
+
+#include "ViewModels/IoStoreInsightsTrack.h"
+#include "ViewModels/IoStoreInsightsTimingViewExtender.h"
+#include "Model/IoStoreInsightsProvider.h"
+
+#include "Algo/BinarySearch.h"
+#include "TraceServices/Model/Callstack.h"
+#include "InsightsCore/Common/TimeUtils.h"
+#include "Insights/ViewModels/TimingTrackViewport.h"
+#include "Insights/ViewModels/TooltipDrawState.h"
+#include "IO/IoChunkId.h"
+
+namespace UE::IoStoreInsights
+{
+	class FIoStoreTimingEvent : public FTimingEvent
+	{
+	public:
+		FIoStoreTimingEvent(const TSharedRef<const FIoStoreInsightsTrack> InTrack, double InStartTime, double InEndTime, uint32 InDepth, const FIoStoreEventState& InEventState)
+			: FTimingEvent(InTrack, InStartTime, InEndTime, InDepth, InEventState.TimingEventType) {}
+		INSIGHTS_DECLARE_RTTI(FIoStoreTimingEvent, FTimingEvent)
+	};
+	INSIGHTS_IMPLEMENT_RTTI(FIoStoreTimingEvent)
+
+
+
+
+
+	FIoStoreInsightsTrack::FIoStoreInsightsTrack(FIoStoreInsightsViewSharedState& InSharedState)
+		: FTimingEventsTrack(TEXT("IoStore Activity"))
+		, SharedState(InSharedState)
+	{
+	}
+
+
+	void FIoStoreInsightsTrack::InitTooltip(FTooltipDrawState& InOutTooltip, const ITimingEvent& InTooltipEvent) const
+	{
+		InOutTooltip.ResetContent();
+
+		if (InTooltipEvent.CheckTrack(this) && InTooltipEvent.Is<FIoStoreTimingEvent>())
+		{
+			const FIoStoreTimingEvent& TooltipEvent = InTooltipEvent.As<FIoStoreTimingEvent>();
+
+			auto MatchEvent = [this, &TooltipEvent](double InStartTime, double InEndTime, uint32 InDepth)
+			{
+				return InDepth == TooltipEvent.GetDepth()
+					&& InStartTime == TooltipEvent.GetStartTime()
+					&& InEndTime == TooltipEvent.GetEndTime();
+			};
+
+			FTimingEventSearchParameters SearchParameters(TooltipEvent.GetStartTime(), TooltipEvent.GetEndTime(), ETimingEventSearchFlags::StopAtFirstMatch, MatchEvent);
+			FindIoStoreEvent(SearchParameters, [this, &InOutTooltip, &TooltipEvent](double InFoundStartTime, double InFoundEndTime, uint32 InFoundDepth, const FIoStoreEventState& InEvent)
+			{
+				const IIoStoreInsightsProvider* IoStoreActivityProvider = SharedState.GetAnalysisSession().ReadProvider<IIoStoreInsightsProvider>(FIoStoreInsightsProvider::ProviderName);
+				const FIoStoreRequest& IoStoreRequest = IoStoreActivityProvider->GetIoStoreRequest(InEvent.RequestState->IoStoreRequestIndex);
+				const EIoStoreActivityType ActivityType = static_cast<EIoStoreActivityType>(InEvent.Type & 0x0F);
+				const bool bHasFailed = ((InEvent.Type & 0xF0) != 0);
+
+				// prepare tooltip title
+				FString TypeStr;
+				uint32 TypeColor;
+				if (bHasFailed)
+				{
+					TypeStr = TEXT("Failed ");
+					TypeStr += LexToString(ActivityType);
+					TypeColor = 0xFFFF3333;
+				}
+				else
+				{
+					TypeStr = LexToString(ActivityType);
+					TypeColor = GetIoStoreActivityTypeColor(ActivityType);
+				}
+				FLinearColor TypeLinearColor = FLinearColor(FColor(TypeColor));
+				TypeLinearColor.R *= 2.0f;
+				TypeLinearColor.G *= 2.0f;
+				TypeLinearColor.B *= 2.0f;
+				InOutTooltip.AddTitle(TypeStr, TypeLinearColor);
+
+				// add tooltip fields
+				const double Duration = InEvent.EndTime - InEvent.StartTime;
+				InOutTooltip.AddNameValueTextLine(TEXT("Duration:"), Insights::FormatTimeAuto(Duration));
+
+				const TCHAR* PackageName = *IoStoreRequest.PackageName ? IoStoreRequest.PackageName : IoStoreActivityProvider->GetPackageName(IoStoreRequest.PackageId);
+				if (*PackageName)
+				{
+					InOutTooltip.AddNameValueTextLine(TEXT("Package:"), PackageName);
+				}
+				if (*IoStoreRequest.ExtraTag)
+				{
+					InOutTooltip.AddNameValueTextLine(TEXT("Tag:"), IoStoreRequest.ExtraTag);
+				}
+				InOutTooltip.AddNameValueTextLine(TEXT("Chunk Type:"), LexToString((EIoChunkType)IoStoreRequest.ChunkType) );
+				InOutTooltip.AddNameValueTextLine(TEXT("Chunk Id Hash:"), *FString::Printf(TEXT("0x%X"), IoStoreRequest.ChunkIdHash ) );
+				InOutTooltip.AddNameValueTextLine(TEXT("Offset:"), FText::AsNumber(IoStoreRequest.Offset).ToString() + TEXT(" bytes"));
+				if (ActivityType == EIoStoreActivityType::Request_Read)
+				{
+					if (IoStoreRequest.Size != std::numeric_limits<uint64>::max() && IoStoreRequest.Size != InEvent.ActualSize)
+					{
+						InOutTooltip.AddNameValueTextLine(TEXT("Size:"), FText::AsNumber(IoStoreRequest.Size).ToString() + TEXT(" bytes"));
+					}
+					InOutTooltip.AddNameValueTextLine(TEXT("Result Size:"), FText::AsNumber(InEvent.ActualSize).ToString() + TEXT(" bytes"));
+					InOutTooltip.AddNameValueTextLine(TEXT("Backend:"), IoStoreActivityProvider->GetBackendName(InEvent.BackendHandle));
+				}
+				else
+				{
+					InOutTooltip.AddNameValueTextLine(TEXT("Requested Size:"), IoStoreRequest.Size != std::numeric_limits<uint64>::max() ? (FText::AsNumber(IoStoreRequest.Size).ToString() + TEXT(" bytes")) : TEXT("(all available data)") );
+				}
+
+				// potentially add a callstack at the very bottom
+				static const bool bAlwaysShowCallstack = false; // @todo: put on the right-click menu
+				if (bAlwaysShowCallstack || (*PackageName == '\0' && *IoStoreRequest.ExtraTag == '\0')) // only show callstack if there's no package or custom tag
+				{
+					// package id may be useful for debugging where the missing tag is
+					if (IoStoreRequest.PackageId != 0)
+					{
+						InOutTooltip.AddNameValueTextLine(TEXT("Package Id:"), FString::Printf(TEXT("0x%llX"), IoStoreRequest.PackageId)); 
+					}
+
+					// append the callstack
+					const TraceServices::ICallstacksProvider* CallstacksProvider = SharedState.GetAnalysisSession().ReadProvider<TraceServices::ICallstacksProvider>(TraceServices::GetCallstacksProviderName());
+					const TraceServices::FCallstack* Callstack = (IoStoreRequest.CallstackId && CallstacksProvider) ? CallstacksProvider->GetCallstack(IoStoreRequest.CallstackId) : nullptr;
+					if (Callstack)
+					{
+						InOutTooltip.AddTextLine(TEXT(""), FTooltipDrawState::DefaultNameColor);
+						for (uint8 CallstackFrame = 0; CallstackFrame < Callstack->Num(); CallstackFrame++)
+						{
+							InOutTooltip.AddTextLine(Callstack->Name(CallstackFrame), FTooltipDrawState::DefaultNameColor);
+						}
+					}
+				}
+
+				InOutTooltip.UpdateLayout();
+			});
+		}
+	}
+
+
+
+	bool FIoStoreInsightsTrack::FindIoStoreEvent(const FTimingEventSearchParameters& InParameters, TFunctionRef<void(double, double, uint32, const FIoStoreEventState&)> InFoundPredicate) const
+	{
+		return TTimingEventSearch<FIoStoreEventState>::Search(
+			InParameters,
+
+			// Search...
+			[this](TTimingEventSearch<FIoStoreEventState>::FContext& InContext)
+			{
+				const TArray<FIoStoreEventState>& Events = SharedState.GetAllEvents();
+
+				// Events are sorted by start time.
+				// Find the first event with StartTime >= searched EndTime.
+				int32 StartIndex = Algo::LowerBoundBy(Events, InContext.GetParameters().EndTime,
+					[](const FIoStoreEventState& Event) { return Event.StartTime; });
+
+				// Start at the last event with StartTime < searched EndTime.
+				for (int32 Index = StartIndex - 1; Index >= 0; --Index)
+				{
+					const FIoStoreEventState& Event = Events[Index];
+
+					if (Event.EndTime <= InContext.GetParameters().StartTime ||
+						Event.StartTime >= InContext.GetParameters().EndTime)
+					{
+						continue;
+					}
+
+					InContext.Check(Event.StartTime, Event.EndTime, Event.Depth, Event);
+
+					if (!InContext.ShouldContinueSearching())
+					{
+						break;
+					}
+				}
+			},
+
+			// Found!
+			[&InFoundPredicate](double InFoundStartTime, double InFoundEndTime, uint32 InFoundDepth, const FIoStoreEventState& InEvent)
+			{
+				InFoundPredicate(InFoundStartTime, InFoundEndTime, InFoundDepth, InEvent);
+			});
+	}
+
+
+
+	void FIoStoreInsightsTrack::AddIoStoreEventToBuilder(const FIoStoreEventState& Event, ITimingEventsTrackDrawStateBuilder& Builder) const
+	{
+		ensure(Event.Depth < FIoStoreInsightsViewSharedState::MaxLanes);
+
+		const EIoStoreActivityType ActivityType = static_cast<EIoStoreActivityType>(Event.Type & 0x0F);
+		const bool bHasFailed = ((Event.Type & 0xF0) != 0);
+		const uint32 Color = bHasFailed ? 0xFFAA0000 : GetIoStoreActivityTypeColor(ActivityType);
+
+		Builder.AddEvent(Event.StartTime, Event.EndTime, Event.Depth, Color,
+			[&Event, ActivityType, bHasFailed, this](float Width)
+			{
+				FString EventName;
+
+				if (bHasFailed)
+				{
+					EventName += TEXT("Failed ");
+				}
+
+				EventName += LexToString(ActivityType);
+
+				const float MinWidth = static_cast<float>(EventName.Len()) * 4.0f + 32.0f;
+				if (Width > MinWidth)
+				{
+					const double Duration = Event.EndTime - Event.StartTime;
+					EventName += TEXT(" (");
+					EventName += UE::Insights::FormatTimeAuto(Duration);
+					EventName += TEXT(")");
+				}
+
+				return EventName;
+			});
+	}
+
+
+
+	void FIoStoreInsightsTrack::BuildDrawState(ITimingEventsTrackDrawStateBuilder& Builder, const ITimingTrackUpdateContext& Context)
+	{
+		const FTimingTrackViewport& Viewport = Context.GetViewport();
+
+		for (const FIoStoreEventState& Event : SharedState.GetAllEvents())
+		{
+			if (Event.EndTime <= Viewport.GetStartTime())
+			{
+				continue;
+			}
+			if (Event.StartTime >= Viewport.GetEndTime())
+			{
+				break;
+			}
+			if (Event.Depth >= FIoStoreInsightsViewSharedState::MaxLanes)
+			{
+				continue;
+			}
+
+			AddIoStoreEventToBuilder(Event, Builder);
+		}
+	}
+
+
+
+	void FIoStoreInsightsTrack::BuildFilteredDrawState(ITimingEventsTrackDrawStateBuilder& Builder, const ITimingTrackUpdateContext& Context)
+	{
+		const TSharedPtr<ITimingEventFilter> EventFilterPtr = Context.GetEventFilter();
+		if (EventFilterPtr.IsValid() && EventFilterPtr->FilterTrack(*this))
+		{
+			const FTimingTrackViewport& Viewport = Context.GetViewport();
+			const IIoStoreInsightsProvider* IoStoreActivityProvider = SharedState.GetAnalysisSession().ReadProvider<IIoStoreInsightsProvider>(FIoStoreInsightsProvider::ProviderName);
+
+			for (const FIoStoreEventState& Event : SharedState.GetAllEvents())
+			{
+				if (Event.EndTime <= Viewport.GetStartTime())
+				{
+					continue;
+				}
+				if (Event.StartTime >= Viewport.GetEndTime())
+				{
+					break;
+				}
+				if (Event.Depth >= FIoStoreInsightsViewSharedState::MaxLanes)
+				{
+					continue;
+				}
+
+				FIoStoreTimingEvent TimingEvent(SharedThis(this), Event.StartTime, Event.EndTime, Event.Depth, Event);
+				if (EventFilterPtr->FilterEvent(TimingEvent))
+				{
+					AddIoStoreEventToBuilder(Event, Builder);
+				}
+			}
+		}
+	}
+
+
+
+	const TSharedPtr<const ITimingEvent> FIoStoreInsightsTrack::SearchEvent(const FTimingEventSearchParameters& InSearchParameters) const
+	{
+		TSharedPtr<const ITimingEvent> FoundEvent;
+
+		FindIoStoreEvent(InSearchParameters, [this, &FoundEvent](double InFoundStartTime, double InFoundEndTime, uint32 InFoundDepth, const FIoStoreEventState& InEventState)
+		{
+			FoundEvent = MakeShared<FIoStoreTimingEvent>(SharedThis(this), InFoundStartTime, InFoundEndTime, InFoundDepth, InEventState);
+		});
+
+		return FoundEvent;
+	}
+
+
+
+	uint32 FIoStoreInsightsTrack::GetIoStoreActivityTypeColor(EIoStoreActivityType Type) const
+	{
+		switch(Type)
+		{
+			case EIoStoreActivityType::Request_Pending: return 0xFF334433;
+			case EIoStoreActivityType::Request_Read:    return 0xFF33AA33;
+		}
+		return 0x55333333;
+	}
+
+} // namespace UE::IoStoreInsights
+
diff --git a/Engine/Plugins/IoStoreInsights/Source/IoStoreInsights/Private/ViewModels/IoStoreInsightsTrack.h b/Engine/Plugins/IoStoreInsights/Source/IoStoreInsights/Private/ViewModels/IoStoreInsightsTrack.h
new file mode 100644
--- /dev/null
+++ b/Engine/Plugins/IoStoreInsights/Source/IoStoreInsights/Private/ViewModels/IoStoreInsightsTrack.h
@@ -0,0 +1,38 @@
+// Copyright Epic Games, Inc. All Rights Reserved.
+
+#pragma once
+
+#include "CoreTypes.h"
+#include "Insights/ViewModels/TimingEvent.h"
+#include "Insights/ViewModels/TimingEventsTrack.h"
+
+class FTimingEventSearchParameters;
+
+namespace UE::IoStoreInsights
+{
+	enum class EIoStoreActivityType : uint8;
+	class FIoStoreInsightsViewSharedState;
+	struct FIoStoreEventState;
+
+	class FIoStoreInsightsTrack : public FTimingEventsTrack
+	{
+	public:
+		explicit FIoStoreInsightsTrack(FIoStoreInsightsViewSharedState& InSharedState);
+		virtual ~FIoStoreInsightsTrack() override {}
+
+		virtual void InitTooltip(FTooltipDrawState& InOutTooltip, const ITimingEvent& InTooltipEvent) const override;
+		virtual void BuildDrawState(ITimingEventsTrackDrawStateBuilder& Builder, const ITimingTrackUpdateContext& Context) override;
+		virtual void BuildFilteredDrawState(ITimingEventsTrackDrawStateBuilder& Builder, const ITimingTrackUpdateContext& Context) override;
+		virtual const TSharedPtr<const ITimingEvent> SearchEvent(const FTimingEventSearchParameters& InSearchParameters) const override;
+
+	protected:
+		void AddIoStoreEventToBuilder(const FIoStoreEventState& Event, ITimingEventsTrackDrawStateBuilder& Builder) const;
+		bool FindIoStoreEvent(const FTimingEventSearchParameters& InParameters, TFunctionRef<void(double, double, uint32, const FIoStoreEventState&)> InFoundPredicate) const;
+
+		uint32 GetIoStoreActivityTypeColor(EIoStoreActivityType Type) const;
+
+	protected:
+		FIoStoreInsightsViewSharedState& SharedState;
+	};
+
+} // namespace UE::IoStoreInsights
diff --git a/Engine/Plugins/IoStoreInsights/Source/IoStoreInsights/Public/IIoStoreInsightsProvider.h b/Engine/Plugins/IoStoreInsights/Source/IoStoreInsights/Public/IIoStoreInsightsProvider.h
new file mode 100644
--- /dev/null
+++ b/Engine/Plugins/IoStoreInsights/Source/IoStoreInsights/Public/IIoStoreInsightsProvider.h
@@ -0,0 +1,62 @@
+// Copyright Epic Games, Inc. All Rights Reserved.
+
+#pragma once
+
+#include "TraceServices/Model/AnalysisSession.h"
+#include "Model/IntervalTimeline.h"
+#include "Templates/Function.h"
+
+namespace UE::IoStoreInsights
+{
+	enum class EIoStoreActivityType : uint8
+	{
+		Request_Pending,
+		Request_Read,
+
+		Count,
+		Invalid = Count
+	};
+
+	struct FIoStoreRequest
+	{
+		uint32 IoStoreRequestIndex = uint32(-1);
+		uint32 ChunkIdHash = 0;
+		uint8 ChunkType = 0;
+		uint64 Offset = 0;
+		uint64 Size = 0;
+		uint64 CallstackId = 0;
+		uint64 PackageId = 0;
+		const TCHAR* PackageName = nullptr;
+		const TCHAR* ExtraTag = nullptr;
+	};
+
+	struct FIoStoreActivity
+	{
+		const FIoStoreRequest* IoStoreRequest = nullptr;
+		double StartTime = 0.0;
+		double EndTime = 0.0;
+		uint64 ActualSize = 0;
+		uint64 BackendHandle = 0;
+		uint32 ThreadId = 0;
+		EIoStoreActivityType ActivityType = EIoStoreActivityType::Invalid;
+		bool Failed = false;
+	};
+
+	class IIoStoreInsightsProvider : public TraceServices::IProvider
+	{
+	public:
+		typedef TraceServices::ITimeline<FIoStoreActivity*> Timeline;
+
+		virtual ~IIoStoreInsightsProvider() = default;
+		virtual void EnumerateIoStoreRequests(TFunctionRef<bool(const FIoStoreRequest&, const Timeline&)> Callback) const = 0;
+
+		virtual const TCHAR* GetPackageName(uint64 PackageId) const = 0;
+		virtual const TCHAR* GetBackendName(uint64 BackendHandle) const = 0;
+		virtual const FIoStoreRequest& GetIoStoreRequest(uint32 IoStoreRequestIndex) const = 0;
+	};
+
+	IOSTOREINSIGHTS_API const IIoStoreInsightsProvider* ReadIoStoreInsightsProvider(const TraceServices::IAnalysisSession& Session);
+
+	IOSTOREINSIGHTS_API const TCHAR* LexToString(EIoStoreActivityType ActivityType);
+
+} // namespace UE::IoStoreInsights
diff --git a/Engine/Programs/UnrealInsights/Config/DefaultEngine.ini b/Engine/Programs/UnrealInsights/Config/DefaultEngine.ini
--- a/Engine/Programs/UnrealInsights/Config/DefaultEngine.ini
+++ b/Engine/Programs/UnrealInsights/Config/DefaultEngine.ini
@@ -4,6 +4,7 @@
 +ProgramEnabledPlugins=TraceDataFilters
 +ProgramEnabledPlugins=SlateInsights
 +ProgramEnabledPlugins=RenderGraphInsights
++ProgramEnabledPlugins=IoStoreInsights
 +ProgramEnabledPlugins=UdpMessaging
 +ProgramEnabledPlugins=TcpMessaging
 +ProgramEnabledPlugins=NetworkPredictionInsights
