[Insights] Fixed spam of log messages comming from the IoStoreInsights plugin.

#rnx

diff --git a/Engine/Plugins/IoStoreInsights/Source/IoStoreInsights/Private/ViewModels/IoStoreInsightsTimingViewExtender.cpp b/Engine/Plugins/IoStoreInsights/Source/IoStoreInsights/Private/ViewModels/IoStoreInsightsTimingViewExtender.cpp
--- a/Engine/Plugins/IoStoreInsights/Source/IoStoreInsights/Private/ViewModels/IoStoreInsightsTimingViewExtender.cpp
+++ b/Engine/Plugins/IoStoreInsights/Source/IoStoreInsights/Private/ViewModels/IoStoreInsightsTimingViewExtender.cpp
@@ -102,14 +102,22 @@ namespace UE::IoStoreInsights
 				});
 			}
 			Stopwatch.Stop();
-			UE_LOG(LogTimingProfiler, Log, TEXT("[IO] Enumerated %s IoStore activities in %s."), *FText::AsNumber(RequestStates.Num()).ToString(), *Insights::FormatTimeAuto(Stopwatch.GetAccumulatedTime()));
+			double EnumeratedTime = Stopwatch.GetAccumulatedTime();
+			if (EnumeratedTime > 0.01)
+			{
+				UE_LOG(LogTimingProfiler, Log, TEXT("[IO] Enumerated IoStore activities (%d request states, %d events) in %s."), RequestStates.Num(), AllIoEvents.Num(), *Insights::FormatTimeAuto(EnumeratedTime));
+			}
 
 			// Sort cached items by Start Time
-			Stopwatch.Restart();			
+			Stopwatch.Restart();
 			RequestStates.Sort([](const TSharedPtr<FIoStoreRequestState>& A, const TSharedPtr<FIoStoreRequestState>& B) { return A->StartTime < B->StartTime; });
 			AllIoEvents.Sort([](const FIoStoreEventState& A, const FIoStoreEventState& B) { return A.StartTime < B.StartTime; });
 			Stopwatch.Stop();
-			UE_LOG(LogTimingProfiler, Log, TEXT("[IO] Sorted IoStore activities in %s."), *Insights::FormatTimeAuto(Stopwatch.GetAccumulatedTime()));
+			double SortTime = Stopwatch.GetAccumulatedTime();
+			if (SortTime > 0.01)
+			{
+				UE_LOG(LogTimingProfiler, Log, TEXT("[IO] Sorted IoStore activities (%d request states, %d events) in %s."), RequestStates.Num(), AllIoEvents.Num(), *Insights::FormatTimeAuto(SortTime));
+			}
 
 			// Compute depth for IoStore activities (avoids overlaps).
 			if (RequestStates.Num() > 0)
