colour IoStore reads based on the backend

#rnx

diff --git a/Engine/Plugins/IoStoreInsights/Source/IoStoreInsights/Private/ViewModels/IoStoreInsightsTrack.cpp b/Engine/Plugins/IoStoreInsights/Source/IoStoreInsights/Private/ViewModels/IoStoreInsightsTrack.cpp
--- a/Engine/Plugins/IoStoreInsights/Source/IoStoreInsights/Private/ViewModels/IoStoreInsightsTrack.cpp
+++ b/Engine/Plugins/IoStoreInsights/Source/IoStoreInsights/Private/ViewModels/IoStoreInsightsTrack.cpp
@@ -180,13 +180,27 @@ namespace UE::IoStoreInsights
 
 
 
-	void FIoStoreInsightsTrack::AddIoStoreEventToBuilder(const FIoStoreEventState& Event, ITimingEventsTrackDrawStateBuilder& Builder) const
+	void FIoStoreInsightsTrack::AddIoStoreEventToBuilder(const FIoStoreEventState& Event, ITimingEventsTrackDrawStateBuilder& Builder, const IIoStoreInsightsProvider* IoStoreActivityProvider) const
 	{
 		ensure(Event.Depth < FIoStoreInsightsViewSharedState::MaxLanes);
 
 		const EIoStoreActivityType ActivityType = static_cast<EIoStoreActivityType>(Event.Type & 0x0F);
 		const bool bHasFailed = ((Event.Type & 0xF0) != 0);
-		const uint32 Color = bHasFailed ? 0xFFAA0000 : GetIoStoreActivityTypeColor(ActivityType);
+		
+		uint32 Color;
+		if (bHasFailed)
+		{
+			Color = 0xFFAA0000;
+		}
+		else if (ActivityType == EIoStoreActivityType::Request_Read)
+		{
+			const TCHAR* BackendName = IoStoreActivityProvider->GetBackendName(Event.BackendHandle);
+			Color =  FTimingEvent::ComputeEventColor(BackendName);
+		}
+		else
+		{
+			Color = GetIoStoreActivityTypeColor(ActivityType);
+		}
 
 		Builder.AddEvent(Event.StartTime, Event.EndTime, Event.Depth, Color,
 			[&Event, ActivityType, bHasFailed, this](float Width)
@@ -217,6 +231,8 @@ namespace UE::IoStoreInsights
 
 	void FIoStoreInsightsTrack::BuildDrawState(ITimingEventsTrackDrawStateBuilder& Builder, const ITimingTrackUpdateContext& Context)
 	{
+		const IIoStoreInsightsProvider* IoStoreActivityProvider = SharedState.GetAnalysisSession().ReadProvider<IIoStoreInsightsProvider>(FIoStoreInsightsProvider::ProviderName);
+
 		const FTimingTrackViewport& Viewport = Context.GetViewport();
 
 		for (const FIoStoreEventState& Event : SharedState.GetAllEvents())
@@ -234,7 +250,7 @@ namespace UE::IoStoreInsights
 				continue;
 			}
 
-			AddIoStoreEventToBuilder(Event, Builder);
+			AddIoStoreEventToBuilder(Event, Builder, IoStoreActivityProvider);
 		}
 	}
 
@@ -266,7 +282,7 @@ namespace UE::IoStoreInsights
 				FIoStoreTimingEvent TimingEvent(SharedThis(this), Event.StartTime, Event.EndTime, Event.Depth, Event);
 				if (EventFilterPtr->FilterEvent(TimingEvent))
 				{
-					AddIoStoreEventToBuilder(Event, Builder);
+					AddIoStoreEventToBuilder(Event, Builder, IoStoreActivityProvider);
 				}
 			}
 		}
diff --git a/Engine/Plugins/IoStoreInsights/Source/IoStoreInsights/Private/ViewModels/IoStoreInsightsTrack.h b/Engine/Plugins/IoStoreInsights/Source/IoStoreInsights/Private/ViewModels/IoStoreInsightsTrack.h
--- a/Engine/Plugins/IoStoreInsights/Source/IoStoreInsights/Private/ViewModels/IoStoreInsightsTrack.h
+++ b/Engine/Plugins/IoStoreInsights/Source/IoStoreInsights/Private/ViewModels/IoStoreInsightsTrack.h
@@ -26,7 +26,7 @@ namespace UE::IoStoreInsights
 		virtual const TSharedPtr<const ITimingEvent> SearchEvent(const FTimingEventSearchParameters& InSearchParameters) const override;
 
 	protected:
-		void AddIoStoreEventToBuilder(const FIoStoreEventState& Event, ITimingEventsTrackDrawStateBuilder& Builder) const;
+		void AddIoStoreEventToBuilder(const FIoStoreEventState& Event, ITimingEventsTrackDrawStateBuilder& Builder, const class IIoStoreInsightsProvider* IoStoreActivityProvider) const;
 		bool FindIoStoreEvent(const FTimingEventSearchParameters& InParameters, TFunctionRef<void(double, double, uint32, const FIoStoreEventState&)> InFoundPredicate) const;
 
 		uint32 GetIoStoreActivityTypeColor(EIoStoreActivityType Type) const;
