[Trivial] Fix Start Time record for IoStore trace activities

#rnx

diff --git a/Engine/Plugins/IoStoreInsights/Source/IoStoreInsights/Private/Model/IoStoreInsightsProvider.cpp b/Engine/Plugins/IoStoreInsights/Source/IoStoreInsights/Private/Model/IoStoreInsightsProvider.cpp
--- a/Engine/Plugins/IoStoreInsights/Source/IoStoreInsights/Private/Model/IoStoreInsightsProvider.cpp
+++ b/Engine/Plugins/IoStoreInsights/Source/IoStoreInsights/Private/Model/IoStoreInsightsProvider.cpp
@@ -64,7 +64,7 @@ namespace UE::IoStoreInsights
 		IoStoreActivity.IoStoreRequest = &IoStoreRequestInfo.IoStoreRequestInfo;
 		IoStoreActivity.BackendHandle = BackendHandle;
 		IoStoreActivity.ActualSize = 0;
-		IoStoreActivity.StartTime = 0;
+		IoStoreActivity.StartTime = Time;
 		IoStoreActivity.EndTime = std::numeric_limits<double>::infinity();
 		IoStoreActivity.ThreadId = ThreadId;
 		IoStoreActivity.ActivityType = Type;
