[Trivial] Move IoStoreInsights ProviderName to the public interface

#rnx

diff --git a/Engine/Plugins/IoStoreInsights/Source/IoStoreInsights/Private/Model/IoStoreInsightsProvider.h b/Engine/Plugins/IoStoreInsights/Source/IoStoreInsights/Private/Model/IoStoreInsightsProvider.h
--- a/Engine/Plugins/IoStoreInsights/Source/IoStoreInsights/Private/Model/IoStoreInsightsProvider.h
+++ b/Engine/Plugins/IoStoreInsights/Source/IoStoreInsights/Private/Model/IoStoreInsightsProvider.h
@@ -15,8 +15,6 @@ namespace UE::IoStoreInsights
 	class FIoStoreInsightsProvider : public IIoStoreInsightsProvider
 	{
 	public:
-		static FName ProviderName;
-
 		struct FTimelineSettings
 		{
 			enum
diff --git a/Engine/Plugins/IoStoreInsights/Source/IoStoreInsights/Public/IIoStoreInsightsProvider.h b/Engine/Plugins/IoStoreInsights/Source/IoStoreInsights/Public/IIoStoreInsightsProvider.h
--- a/Engine/Plugins/IoStoreInsights/Source/IoStoreInsights/Public/IIoStoreInsightsProvider.h
+++ b/Engine/Plugins/IoStoreInsights/Source/IoStoreInsights/Public/IIoStoreInsightsProvider.h
@@ -45,6 +45,8 @@ namespace UE::IoStoreInsights
 	class IIoStoreInsightsProvider : public TraceServices::IProvider
 	{
 	public:
+		static IOSTOREINSIGHTS_API FName ProviderName;
+
 		typedef TraceServices::ITimeline<FIoStoreActivity*> Timeline;
 
 		virtual ~IIoStoreInsightsProvider() = default;
