[Trivial] add missing file from previous checkin - fix insights provider name.

#rnx

diff --git a/Engine/Plugins/IoStoreInsights/Source/IoStoreInsights/Private/Model/IoStoreInsightsProvider.cpp b/Engine/Plugins/IoStoreInsights/Source/IoStoreInsights/Private/Model/IoStoreInsightsProvider.cpp
--- a/Engine/Plugins/IoStoreInsights/Source/IoStoreInsights/Private/Model/IoStoreInsightsProvider.cpp
+++ b/Engine/Plugins/IoStoreInsights/Source/IoStoreInsights/Private/Model/IoStoreInsightsProvider.cpp
@@ -4,7 +4,7 @@
 
 namespace UE::IoStoreInsights
 {
-	FName FIoStoreInsightsProvider::ProviderName("IoStoreProvider");
+	FName IIoStoreInsightsProvider::ProviderName("IoStoreProvider");
 
 
 
