Adding custom analysis tab for IoStore insights

#rnx

diff --git a/Engine/Plugins/IoStoreInsights/Source/IoStoreInsights/IoStoreInsights.Build.cs b/Engine/Plugins/IoStoreInsights/Source/IoStoreInsights/IoStoreInsights.Build.cs
--- a/Engine/Plugins/IoStoreInsights/Source/IoStoreInsights/IoStoreInsights.Build.cs
+++ b/Engine/Plugins/IoStoreInsights/Source/IoStoreInsights/IoStoreInsights.Build.cs
@@ -12,6 +12,7 @@ namespace UnrealBuildTool.Rules
 				"Core",
 				"SlateCore",
 				"Slate",
+				"InputCore",
 				"TraceAnalysis",
 				"TraceServices",
 				"TraceInsights",
diff --git a/Engine/Plugins/IoStoreInsights/Source/IoStoreInsights/Private/IoStoreInsightsModule.cpp b/Engine/Plugins/IoStoreInsights/Source/IoStoreInsights/Private/IoStoreInsightsModule.cpp
--- a/Engine/Plugins/IoStoreInsights/Source/IoStoreInsights/Private/IoStoreInsightsModule.cpp
+++ b/Engine/Plugins/IoStoreInsights/Source/IoStoreInsights/Private/IoStoreInsightsModule.cpp
@@ -4,21 +4,93 @@
 #include "Modules/ModuleManager.h"
 #include "Features/IModularFeatures.h"
 
+#include "Insights/IUnrealInsightsModule.h"
+#include "Widgets/SIoStoreAnalysisTab.h"
+#include "Widgets/Docking/SDockTab.h"
+
+#define LOCTEXT_NAMESPACE "IoStoreInsightsModule"
+
 namespace UE::IoStoreInsights
 {
+	namespace Private
+	{
+		const FName ViewTab("IoStoreViewTab");
+	}
+
+
+
+	FIoStoreInsightsModule& FIoStoreInsightsModule::Get()
+	{
+		return FModuleManager::LoadModuleChecked<FIoStoreInsightsModule>("IoStoreInsights");
+	}
+
+
+
 	void FIoStoreInsightsModule::StartupModule()
 	{
 		IModularFeatures::Get().RegisterModularFeature(TraceServices::ModuleFeatureName, &TraceModule);
 		IModularFeatures::Get().RegisterModularFeature(Insights::Timing::TimingViewExtenderFeatureName, &TimingViewExtender);
+
+		IUnrealInsightsModule& UnrealInsightsModule = FModuleManager::LoadModuleChecked<IUnrealInsightsModule>("TraceInsights");
+		FOnRegisterMajorTabExtensions& TimingProfilerLayoutExtension = UnrealInsightsModule.OnRegisterMajorTabExtension(FInsightsManagerTabs::TimingProfilerTabId);
+		TimingProfilerLayoutExtension.AddRaw(this, &FIoStoreInsightsModule::RegisterTimingProfilerLayoutExtensions);
 	}
 
+
+
 	void FIoStoreInsightsModule::ShutdownModule()
 	{
+		if (IUnrealInsightsModule* UnrealInsightsModule = FModuleManager::GetModulePtr<IUnrealInsightsModule>("TraceInsights"))
+		{
+			FOnRegisterMajorTabExtensions& TimingProfilerLayoutExtension = UnrealInsightsModule->OnRegisterMajorTabExtension(FInsightsManagerTabs::TimingProfilerTabId);
+			TimingProfilerLayoutExtension.RemoveAll(this);
+		}
+
 		IModularFeatures::Get().UnregisterModularFeature(TraceServices::ModuleFeatureName, &TraceModule);
 		IModularFeatures::Get().UnregisterModularFeature(Insights::Timing::TimingViewExtenderFeatureName, &TimingViewExtender);
 	}
 
 
+
+	void FIoStoreInsightsModule::RegisterTimingProfilerLayoutExtensions(FInsightsMajorTabExtender& InOutExtender)
+	{
+		InsightsTabManager = InOutExtender.GetTabManager();
+
+		FMinorTabConfig& MinorTabConfig = InOutExtender.AddMinorTabConfig();
+		MinorTabConfig.TabId = Private::ViewTab;
+		MinorTabConfig.TabLabel = LOCTEXT("IoStoreTabTitle", "IoStore View");
+		MinorTabConfig.TabTooltip = LOCTEXT("IoStoreTabTitleTooltip", "Opens the IoStore View tab, allows for diagnostics of IoStore data.");
+		MinorTabConfig.TabIcon = FSlateIcon(FAppStyle::Get().GetStyleSetName(), "Icons.Plugin.TreeItem");
+		MinorTabConfig.WorkspaceGroup = InOutExtender.GetWorkspaceGroup();
+		MinorTabConfig.OnSpawnTab = FOnSpawnTab::CreateLambda([this](const FSpawnTabArgs& Args)
+		{
+			return SNew(SDockTab)
+				.ShouldAutosize(false)
+				.TabRole(ETabRole::PanelTab)
+				[
+					SAssignNew(IoStoreAnalysisView, SIoStoreAnalysisTab)
+				];
+		});
+
+		InOutExtender.GetLayoutExtender().ExtendLayout(Private::ViewTab, ELayoutExtensionPosition::Before, FTabManager::FTab(Private::ViewTab, ETabState::ClosedTab));
+	}
+
+
+
+	TSharedPtr<SIoStoreAnalysisTab> FIoStoreInsightsModule::GetIoStoreAnalysisViewTab(bool bInvoke)
+	{
+		if (bInvoke)
+		{
+			if (TSharedPtr<FTabManager> TabManagerPin = InsightsTabManager.Pin())
+			{
+				TabManagerPin->TryInvokeTab(Private::ViewTab);
+			}
+		}
+		return IoStoreAnalysisView.Pin();
+	}
+
 } //namespace UE::IoStoreInsights
 
 IMPLEMENT_MODULE(UE::IoStoreInsights::FIoStoreInsightsModule, IoStoreInsights);
+
+#undef LOCTEXT_NAMESPACE
diff --git a/Engine/Plugins/IoStoreInsights/Source/IoStoreInsights/Private/IoStoreInsightsModule.h b/Engine/Plugins/IoStoreInsights/Source/IoStoreInsights/Private/IoStoreInsightsModule.h
--- a/Engine/Plugins/IoStoreInsights/Source/IoStoreInsights/Private/IoStoreInsightsModule.h
+++ b/Engine/Plugins/IoStoreInsights/Source/IoStoreInsights/Private/IoStoreInsightsModule.h
@@ -6,15 +6,28 @@
 #include "IoStoreInsightsTraceModule.h"
 #include "ViewModels/IoStoreInsightsTimingViewExtender.h"
 
+struct FInsightsMajorTabExtender;
+class FTabManager;
+
 namespace UE::IoStoreInsights
 {
+	class SIoStoreAnalysisTab;
+
 	class FIoStoreInsightsModule : public IModuleInterface
 	{
 	public:
+		static FIoStoreInsightsModule& Get();
+
 		virtual void StartupModule() override;
 		virtual void ShutdownModule() override;
 
+		TSharedPtr<SIoStoreAnalysisTab> GetIoStoreAnalysisViewTab(bool bInvoke);
+
 	private:
+		void RegisterTimingProfilerLayoutExtensions(FInsightsMajorTabExtender& InOutExtender);
+		TWeakPtr<SIoStoreAnalysisTab> IoStoreAnalysisView;
+		TWeakPtr<FTabManager> InsightsTabManager;
+
 		FIoStoreInsightsTraceModule TraceModule;
 		FIoStoreInsightsTimingViewExtender TimingViewExtender;
 	};
diff --git a/Engine/Plugins/IoStoreInsights/Source/IoStoreInsights/Private/ViewModels/IoStoreInsightsTimingViewExtender.cpp b/Engine/Plugins/IoStoreInsights/Source/IoStoreInsights/Private/ViewModels/IoStoreInsightsTimingViewExtender.cpp
--- a/Engine/Plugins/IoStoreInsights/Source/IoStoreInsights/Private/ViewModels/IoStoreInsightsTimingViewExtender.cpp
+++ b/Engine/Plugins/IoStoreInsights/Source/IoStoreInsights/Private/ViewModels/IoStoreInsightsTimingViewExtender.cpp
@@ -3,11 +3,13 @@
 #include "ViewModels/IoStoreInsightsTimingViewExtender.h"
 #include "ViewModels/IoStoreInsightsTrack.h"
 #include "Model/IoStoreInsightsProvider.h"
+#include "Widgets/SIoStoreAnalysisTab.h"
 #include "Framework/MultiBox/MultiBoxBuilder.h"
 #include "InsightsCore/Common/Stopwatch.h"
 #include "InsightsCore/Common/TimeUtils.h"
 #include "Insights/ITimingViewSession.h"
 #include "Insights/TimingProfilerCommon.h"
+#include "IoStoreInsightsModule.h"
 
 #define LOCTEXT_NAMESPACE "IoStoreViewExtender"
 
@@ -105,7 +107,7 @@ namespace UE::IoStoreInsights
 			double EnumeratedTime = Stopwatch.GetAccumulatedTime();
 			if (EnumeratedTime > 0.01)
 			{
-				UE_LOG(LogTimingProfiler, Log, TEXT("[IO] Enumerated IoStore activities (%d request states, %d events) in %s."), RequestStates.Num(), AllIoEvents.Num(), *Insights::FormatTimeAuto(EnumeratedTime));
+				UE_LOG(LogTimingProfiler, Verbose, TEXT("[IO] Enumerated IoStore activities (%d request states, %d events) in %s."), RequestStates.Num(), AllIoEvents.Num(), *Insights::FormatTimeAuto(EnumeratedTime));
 			}
 
 			// Sort cached items by Start Time
@@ -116,7 +118,7 @@ namespace UE::IoStoreInsights
 			double SortTime = Stopwatch.GetAccumulatedTime();
 			if (SortTime > 0.01)
 			{
-				UE_LOG(LogTimingProfiler, Log, TEXT("[IO] Sorted IoStore activities (%d request states, %d events) in %s."), RequestStates.Num(), AllIoEvents.Num(), *Insights::FormatTimeAuto(SortTime));
+				UE_LOG(LogTimingProfiler, Verbose, TEXT("[IO] Sorted IoStore activities (%d request states, %d events) in %s."), RequestStates.Num(), AllIoEvents.Num(), *Insights::FormatTimeAuto(SortTime));
 			}
 
 			// Compute depth for IoStore activities (avoids overlaps).
@@ -181,7 +183,7 @@ namespace UE::IoStoreInsights
 					}
 				}
 				Stopwatch.Stop();
-				UE_LOG(LogTimingProfiler, Log, TEXT("[IO] Computed layout for IoStore activities in %s."), *Insights::FormatTimeAuto(Stopwatch.GetAccumulatedTime()));
+				UE_LOG(LogTimingProfiler, Verbose, TEXT("[IO] Computed layout for IoStore activities in %s."), *Insights::FormatTimeAuto(Stopwatch.GetAccumulatedTime()));
 
 				Stopwatch.Restart();
 				for (FIoStoreEventState& Event : AllIoEvents)
@@ -190,7 +192,7 @@ namespace UE::IoStoreInsights
 					ensure(Event.Depth < MaxLanes);
 				}
 				Stopwatch.Stop();
-				UE_LOG(LogTimingProfiler, Log, TEXT("[IO] Updated depth for events in %s."), *Insights::FormatTimeAuto(Stopwatch.GetAccumulatedTime()));
+				UE_LOG(LogTimingProfiler, Verbose, TEXT("[IO] Updated depth for events in %s."), *Insights::FormatTimeAuto(Stopwatch.GetAccumulatedTime()));
 			}
 		}
 	}
@@ -226,6 +228,12 @@ namespace UE::IoStoreInsights
 
 			PerSessionDataMap.Remove(&InSession);
 		}
+
+		const bool bInvoke = false;
+		if (TSharedPtr<SIoStoreAnalysisTab> AnalysisView = FIoStoreInsightsModule::Get().GetIoStoreAnalysisViewTab(bInvoke))
+		{
+			AnalysisView->SetSession(nullptr, nullptr, nullptr);
+		}
 	}
 
 
@@ -272,6 +280,14 @@ namespace UE::IoStoreInsights
 		}
 
 		SharedState.Tick(InAnalysisSession);
+
+		// update tab
+		const bool bInvoke = false;
+		if (TSharedPtr<SIoStoreAnalysisTab> AnalysisView = FIoStoreInsightsModule::Get().GetIoStoreAnalysisViewTab(bInvoke))
+		{
+			AnalysisView->SetSession(&InSession, &InAnalysisSession, &SharedState);
+		}
+
 	}
 
 
diff --git a/Engine/Plugins/IoStoreInsights/Source/IoStoreInsights/Private/Widgets/SIoStoreAnalysisActivityView.cpp b/Engine/Plugins/IoStoreInsights/Source/IoStoreInsights/Private/Widgets/SIoStoreAnalysisActivityView.cpp
new file mode 100644
--- /dev/null
+++ b/Engine/Plugins/IoStoreInsights/Source/IoStoreInsights/Private/Widgets/SIoStoreAnalysisActivityView.cpp
@@ -0,0 +1,354 @@
+// Copyright Epic Games, Inc. All Rights Reserved.
+
+#include "SIoStoreAnalysisActivityView.h"
+#include "Model/IoStoreInsightsProvider.h"
+#include "Widgets/Text/STextBlock.h"
+#include "Widgets/Views/SHeaderRow.h"
+#include "Widgets/Views/SExpanderArrow.h"
+#include "InsightsCore/Common/TimeUtils.h"
+#include "IO/IoChunkId.h"
+
+#define LOCTEXT_NAMESPACE "SIoStoreAnalysisActivityView"
+
+namespace UE::IoStoreInsights
+{
+	namespace Private
+	{
+		const FName ColumnRequestPackage("Package");
+		const FName ColumnRequestOffset("Offset");
+		const FName ColumnRequestSize("Size");
+		const FName ColumnRequestDuration("Duration");
+		const FName ColumnRequestChunkId("ChunkId");
+		const FName ColumnRequestChunkType("ChunkType");
+		const FName ColumnRequestStartTime("StartTime");
+
+		// anything in this range will be tinted yellow to red
+		const double HeatColorMinDuration = 0.005;
+		const double HeatColorMaxDuration = 0.01;
+
+		struct SActivityInfoRow : public SMultiColumnTableRow<TSharedPtr<Private::FActivityInfo>>
+		{
+			SLATE_BEGIN_ARGS(SActivityInfoRow) {}
+			SLATE_END_ARGS()
+
+			void Construct(const FArguments& InArgs, const TSharedRef<STableViewBase>& InOwnerTable, TSharedPtr<Private::FActivityInfo> InInfo, const IIoStoreInsightsProvider* InIoStoreActivityProvider)
+			{
+				Info = InInfo;
+				IoStoreActivityProvider = InIoStoreActivityProvider;
+				SMultiColumnTableRow<TSharedPtr<Private::FActivityInfo>>::Construct(FSuperRowType::FArguments(), InOwnerTable);
+			}
+
+			TSharedRef<STextBlock> MakeColumnText(const FText& InText)
+			{
+				return SNew(STextBlock)
+					.Text(InText)
+					.TextStyle(FAppStyle::Get(), "TreeTable.Tooltip");
+			}
+
+			TSharedRef<STextBlock> MakeColumnTextWithHeatColor(const FText& InText, double Value, double MinValue, double MaxValue)
+			{
+				TSharedRef<STextBlock> TextBlock = MakeColumnText(InText);
+				if (Value >= MinValue)
+				{
+					float Pct = FMath::GetRangePct(MinValue, MaxValue, Value);
+					FLinearColor HeatColor = FMath::Lerp(FLinearColor::Yellow, FLinearColor::Red, Pct);
+					TextBlock->SetColorAndOpacity(HeatColor);
+				}
+
+				return TextBlock;
+			}
+
+
+			TSharedRef<SWidget> GenerateWidgetForColumn(const FName& Column)
+			{
+				const FIoStoreRequest* IoStoreRequest = Info->IoStoreActivity->IoStoreRequest;
+
+				if (Info->bIsParent)
+				{
+					// parent-only columns
+					const Private::FActivityInfoParent* Parent = (const Private::FActivityInfoParent*)Info.Get();
+					if (Column == Private::ColumnRequestPackage)
+					{
+						return SNew(SHorizontalBox)
+							+SHorizontalBox::Slot()
+							.AutoWidth()
+							.HAlign(HAlign_Right)
+							.VAlign(VAlign_Fill)
+							[
+								SNew(SExpanderArrow, SharedThis(this) )
+								.StyleSet(ExpanderStyleSet)
+							]
+
+							+SHorizontalBox::Slot()
+							.FillWidth(1)
+							[
+								SNew(STextBlock)
+									.Text(FText::FromString(Parent->DisplayName))
+							];
+					}
+					else if (Column == Private::ColumnRequestOffset && Parent->Children.Num() == 0)
+					{
+						return MakeColumnText(FText::AsNumber(IoStoreRequest->Offset));
+					}
+					else if (Column == Private::ColumnRequestSize)
+					{
+						return MakeColumnText(FText::AsMemory(Parent->TotalSize));
+					}
+					else if (Column == Private::ColumnRequestDuration)
+					{
+						return MakeColumnTextWithHeatColor(FText::FromString(*UE::Insights::FormatTimeAuto(Parent->Duration)), Parent->Duration, HeatColorMinDuration, HeatColorMaxDuration);
+					}
+					else if (Column == Private::ColumnRequestChunkId && Parent->Children.Num() == 0)
+					{
+						return MakeColumnText(FText::FromString(FString::Printf(TEXT("0x%X"), IoStoreRequest->ChunkIdHash)));
+					}
+					else if (Column == Private::ColumnRequestChunkType)
+					{
+						return MakeColumnText(FText::FromString(LexToString((EIoChunkType)IoStoreRequest->ChunkType)));
+					}
+				}
+				else
+				{
+					// child-only columns
+					if (Column == Private::ColumnRequestOffset)
+					{
+						return MakeColumnText(FText::AsNumber(IoStoreRequest->Offset));
+					}
+					else if (Column == Private::ColumnRequestSize)
+					{
+						return MakeColumnText(FText::AsMemory(Info->IoStoreActivity->ActualSize));
+					}
+					else if (Column == Private::ColumnRequestDuration)
+					{
+						return MakeColumnTextWithHeatColor(FText::FromString(*UE::Insights::FormatTimeAuto(Info->Duration)), Info->Duration, HeatColorMinDuration, HeatColorMaxDuration);
+					}
+					else if (Column == Private::ColumnRequestChunkId)
+					{
+						return MakeColumnText(FText::FromString(FString::Printf(TEXT("0x%X"), IoStoreRequest->ChunkIdHash)));
+					}
+				}
+
+				// common columns
+				if (Column == Private::ColumnRequestStartTime)
+				{
+					return MakeColumnText(FText::FromString(*UE::Insights::FormatTimeAuto(Info->IoStoreActivity->StartTime)));
+				}
+
+				return SNullWidget::NullWidget;
+			}
+
+		private:
+			TSharedPtr<Private::FActivityInfo> Info;
+			const IIoStoreInsightsProvider* IoStoreActivityProvider = nullptr;
+		};
+
+	}
+
+
+	void SIoStoreAnalysisActivityView::Construct(const FArguments& InArgs)
+	{
+		SortColumn = Private::ColumnRequestStartTime; // we know the events are generated from the timeline, so they'll be sorted by time as default
+
+		ItemsPtr = InArgs._TreeItemsSource;
+
+		Super::Construct(
+			Super::FArguments()
+			.TreeItemsSource(ItemsPtr)
+			.OnGetChildren(this, &SIoStoreAnalysisActivityView::OnGetChildren)
+			.OnGenerateRow(this, &SIoStoreAnalysisActivityView::OnGenerateRow)
+			.HeaderRow
+			(
+				SAssignNew(HeaderRow, SHeaderRow)
+			)
+		);
+
+		auto AddColumn = [this]( FName ColumnId, const FText& Label, float FillWidth)
+		{
+			HeaderRow->AddColumn( SHeaderRow::FColumn::FArguments()
+				.ColumnId(ColumnId)
+				.DefaultLabel(Label)
+				.FillWidth(FillWidth)
+				.SortMode(this, &SIoStoreAnalysisActivityView::GetColumnSortMode, ColumnId)
+				.OnSort(this, &SIoStoreAnalysisActivityView::OnColumnSortModeChanged)
+			);
+		};
+
+		AddColumn(Private::ColumnRequestPackage,   LOCTEXT("ColumnRequestPackage",   "Package"),    0.4f);
+		AddColumn(Private::ColumnRequestOffset,    LOCTEXT("ColumnRequestOffset",    "Offset"),     0.1f);
+		AddColumn(Private::ColumnRequestSize,      LOCTEXT("ColumnRequestSize",      "Size"),       0.1f);
+		AddColumn(Private::ColumnRequestDuration,  LOCTEXT("ColumnRequestDuration",  "Duration"),   0.1f);
+		AddColumn(Private::ColumnRequestChunkId,   LOCTEXT("ColumnRequestChunkId",   "ChunkId"),    0.1f);
+		AddColumn(Private::ColumnRequestChunkType, LOCTEXT("ColumnRequestChunkType", "Chunk Type"), 0.1f);
+		AddColumn(Private::ColumnRequestStartTime, LOCTEXT("ColumnRequestStartTime", "Time"),       0.1f);
+	}
+
+
+
+	TSharedRef<ITableRow> SIoStoreAnalysisActivityView::OnGenerateRow(TSharedPtr<Private::FActivityInfo> Item, const TSharedRef<STableViewBase>& OwnerTable)
+	{
+		const IIoStoreInsightsProvider* IoStoreActivityProvider = AnalysisSession ? AnalysisSession->ReadProvider<IIoStoreInsightsProvider>(FIoStoreInsightsProvider::ProviderName) : nullptr;
+		return SNew(Private::SActivityInfoRow, OwnerTable, Item, IoStoreActivityProvider);
+	}
+
+
+
+	void SIoStoreAnalysisActivityView::OnGetChildren(TSharedPtr<Private::FActivityInfo> Item, TArray<TSharedPtr<Private::FActivityInfo>>& OutChildren)
+	{
+		if (Item->bIsParent)
+		{
+			const Private::FActivityInfoParent* Parent = (const Private::FActivityInfoParent*)Item.Get();
+			OutChildren.Append(Parent->Children);
+		}
+	}
+
+
+	EColumnSortMode::Type SIoStoreAnalysisActivityView::GetColumnSortMode(const FName ColumnId) const
+	{
+		if (ColumnId == SortColumn)
+		{
+			return SortMode;
+		}
+	
+		return EColumnSortMode::None;
+	}
+
+	void SIoStoreAnalysisActivityView::OnColumnSortModeChanged(const EColumnSortPriority::Type SortPriority, const FName & ColumnId, const EColumnSortMode::Type InSortMode)
+	{
+		// store current sort mode
+		SortColumn = ColumnId;
+		SortMode = InSortMode;
+
+		ApplySort();
+		RebuildList();
+	}
+
+
+	void SIoStoreAnalysisActivityView::ApplySort()
+	{
+		// sort the parent items
+		ApplySort(*ItemsPtr, true);
+
+		// sort their children
+		for (TSharedPtr<Private::FActivityInfo> Info : *ItemsPtr)
+		{
+			check (Info->bIsParent);
+			Private::FActivityInfoParent* Parent = (Private::FActivityInfoParent*)Info.Get();
+			if (Parent->Children.Num() > 0)
+			{
+				ApplySort(Parent->Children, false);
+			}
+		}
+	}
+
+	void SIoStoreAnalysisActivityView::ApplySort( TArray<TSharedPtr<Private::FActivityInfo>>& Items, bool bIsParent )
+	{
+		// helper lambdas to make sorting functions easier to read
+		auto ToParentInfo = []( const TSharedPtr<Private::FActivityInfo> Info )
+		{
+			check(Info->bIsParent);
+			return (const Private::FActivityInfoParent*)Info.Get();
+		};
+
+		auto ToActivity = []( const TSharedPtr<Private::FActivityInfo> Info )
+		{
+			return Info->IoStoreActivity;
+		};
+
+		auto ToRequest = []( const TSharedPtr<Private::FActivityInfo> Info )
+		{
+			return Info->IoStoreActivity->IoStoreRequest;
+		};
+
+		auto HasChildren = [ToParentInfo]( const TSharedPtr<Private::FActivityInfo> Info )
+		{
+			return Info->bIsParent ? (ToParentInfo(Info)->Children.Num() > 0) : false;
+		};
+
+		bool bSortAscending = (SortMode == EColumnSortMode::Ascending);
+		auto CompareForSort = [bSortAscending]( auto A, auto B)
+		{
+			return bSortAscending ? (A < B) : (A > B);
+		};
+
+		// sort by the desired column
+		if (SortColumn == Private::ColumnRequestPackage)
+		{
+			if (bIsParent)
+			{
+				Items.Sort( [=](auto A, auto B)
+				{
+					return CompareForSort(ToParentInfo(A)->DisplayName, ToParentInfo(B)->DisplayName);
+				});
+			}
+			else
+			{
+				Items.Sort( [=](auto A, auto B)
+				{
+					return CompareForSort(ToActivity(A)->StartTime, ToActivity(B)->StartTime); // sort children by time (default) when sorting parent by package names
+				});
+			}
+		}
+		else if (SortColumn == Private::ColumnRequestOffset)
+		{
+			Items.Sort( [=](auto A, auto B)
+			{
+				return (HasChildren(A) == HasChildren(B) ? CompareForSort(ToRequest(A)->Offset, ToRequest(B)->Offset) : (HasChildren(A) == bSortAscending) );
+			});
+		}
+		else if (SortColumn == Private::ColumnRequestSize)
+		{
+			if (bIsParent)
+			{
+				Items.Sort( [=](auto A, auto B)
+				{
+					return CompareForSort(ToParentInfo(A)->TotalSize, ToParentInfo(B)->TotalSize);
+				});
+			}
+			else
+			{
+				Items.Sort( [=](auto A, auto B)
+				{
+					return CompareForSort(ToActivity(A)->ActualSize, ToActivity(B)->ActualSize);
+				});
+			}
+		}
+		else if (SortColumn == Private::ColumnRequestDuration)
+		{
+			Items.Sort( [=](auto A, auto B)
+			{
+				return CompareForSort(A->Duration, B->Duration);
+			});
+		}
+		else if (SortColumn == Private::ColumnRequestChunkId)
+		{
+			Items.Sort( [=](auto A, auto B)
+			{
+				return (HasChildren(A) == HasChildren(B) ? CompareForSort(ToRequest(A)->ChunkIdHash, ToRequest(B)->ChunkIdHash) : (HasChildren(A) == bSortAscending) );
+			});
+		}
+		else if (SortColumn == Private::ColumnRequestChunkType)
+		{
+			Items.Sort( [=](auto A, auto B)
+			{
+				return CompareForSort(ToRequest(A)->ChunkType, ToRequest(B)->ChunkType);
+			});
+		}
+		else if (SortColumn == Private::ColumnRequestStartTime)
+		{
+			Items.Sort( [=](auto A, auto B)
+			{
+				return CompareForSort(ToActivity(A)->StartTime, ToActivity(B)->StartTime);
+			});
+		}
+	}
+
+
+
+	void SIoStoreAnalysisActivityView::SetAnalysisSession( const TraceServices::IAnalysisSession* InAnalysisSession )
+	{
+		AnalysisSession = InAnalysisSession;
+	}
+
+} //namespace UE::IoStoreInsights
+
+#undef LOCTEXT_NAMESPACE
diff --git a/Engine/Plugins/IoStoreInsights/Source/IoStoreInsights/Private/Widgets/SIoStoreAnalysisActivityView.h b/Engine/Plugins/IoStoreInsights/Source/IoStoreInsights/Private/Widgets/SIoStoreAnalysisActivityView.h
new file mode 100644
--- /dev/null
+++ b/Engine/Plugins/IoStoreInsights/Source/IoStoreInsights/Private/Widgets/SIoStoreAnalysisActivityView.h
@@ -0,0 +1,68 @@
+// Copyright Epic Games, Inc. All Rights Reserved.
+
+#pragma once
+
+#include "Widgets/Views/STreeView.h"
+#include "Widgets/Views/SHeaderRow.h"
+
+namespace TraceServices
+{
+	class IAnalysisSession;
+}
+
+namespace UE::IoStoreInsights
+{
+	struct FIoStoreActivity;
+
+	namespace Private
+	{
+		struct FActivityInfo
+		{
+			const FIoStoreActivity* IoStoreActivity = nullptr;
+			double Duration = 0;
+			bool bIsParent = false;
+		};
+
+		struct FActivityInfoChild : public FActivityInfo
+		{
+		};
+
+		struct FActivityInfoParent : public FActivityInfo
+		{
+			FString DisplayName;
+			uint64 TotalSize = 0;
+			TArray<TSharedPtr<FActivityInfo>> Children;
+		};
+	}
+
+	class SIoStoreAnalysisActivityView : public STreeView<TSharedPtr<Private::FActivityInfo>>
+	{
+		typedef STreeView<TSharedPtr<Private::FActivityInfo>> Super;
+
+	public:
+		SLATE_BEGIN_ARGS(SIoStoreAnalysisActivityView) {}
+			SLATE_ARGUMENT(TArray<TSharedPtr<Private::FActivityInfo>>*, TreeItemsSource)
+		SLATE_END_ARGS()
+		~SIoStoreAnalysisActivityView() = default;
+
+		void Construct(const FArguments& InArgs);
+		void SetAnalysisSession(const TraceServices::IAnalysisSession* InAnalysisSession);
+		void ApplySort();
+
+	private:
+		const TraceServices::IAnalysisSession* AnalysisSession = nullptr;
+		TArray<TSharedPtr<Private::FActivityInfo>>* ItemsPtr;
+		FName SortColumn;
+		EColumnSortMode::Type SortMode = EColumnSortMode::Ascending;
+
+		TSharedRef<ITableRow> OnGenerateRow(TSharedPtr<Private::FActivityInfo> Item, const TSharedRef<STableViewBase>& OwnerTable);
+		void OnGetChildren(TSharedPtr<Private::FActivityInfo> Item, TArray<TSharedPtr<Private::FActivityInfo>>& OutChildren);
+		EColumnSortMode::Type GetColumnSortMode(const FName ColumnId) const;
+		void OnColumnSortModeChanged(const EColumnSortPriority::Type SortPriority, const FName & ColumnId, const EColumnSortMode::Type InSortMode);
+
+		void ApplySort( TArray<TSharedPtr<Private::FActivityInfo>>& Items, bool bIsParent );
+
+	};
+
+} //namespace UE::IoStoreInsights
+
diff --git a/Engine/Plugins/IoStoreInsights/Source/IoStoreInsights/Private/Widgets/SIoStoreAnalysisReadSizeHistogramView.cpp b/Engine/Plugins/IoStoreInsights/Source/IoStoreInsights/Private/Widgets/SIoStoreAnalysisReadSizeHistogramView.cpp
new file mode 100644
--- /dev/null
+++ b/Engine/Plugins/IoStoreInsights/Source/IoStoreInsights/Private/Widgets/SIoStoreAnalysisReadSizeHistogramView.cpp
@@ -0,0 +1,137 @@
+// Copyright Epic Games, Inc. All Rights Reserved.
+
+#include "SIoStoreAnalysisReadSizeHistogramView.h"
+#include "Widgets/Text/STextBlock.h"
+#include "Widgets/Views/SHeaderRow.h"
+#include "Widgets/Notifications/SProgressBar.h"
+#include "InsightsCore/Common/TimeUtils.h"
+
+#define LOCTEXT_NAMESPACE "SIoStoreAnalysisReadSizeHistogramView"
+
+namespace UE::IoStoreInsights
+{
+	namespace Private
+	{
+		const FName ColumnQuantizedReadSize("ReadSize");
+		const FName ColumnQuantizedReadCountGraph("ReadCountGraph");
+		const FName ColumnQuantizedReadCount("ReadCount");
+		const FName ColumnQuantizedReadTime("ReadTime");
+
+		struct SReadSizeHistogramViewRow : public SMultiColumnTableRow<TSharedPtr<Private::FReadSizeHistogramItem>>
+		{
+			SLATE_BEGIN_ARGS(SReadSizeHistogramViewRow) {}
+			SLATE_END_ARGS()
+
+			void Construct(const FArguments& InArgs, const TSharedRef<STableViewBase>& InOwnerTable, TSharedPtr<Private::FReadSizeHistogramItem> InItem)
+			{
+				Item = InItem;
+				SMultiColumnTableRow<TSharedPtr<Private::FReadSizeHistogramItem>>::Construct(FSuperRowType::FArguments(), InOwnerTable);
+			}
+
+			TSharedRef<SWidget> GenerateWidgetForColumn(const FName& Column)
+			{
+				if (Column == Private::ColumnQuantizedReadSize)
+				{
+					return SNew(STextBlock)
+						.Text(FText::AsMemory(Item->QuantizedReadSize));
+				}
+				if (Column == Private::ColumnQuantizedReadCountGraph)
+				{
+					return SNew(SHorizontalBox)
+						+SHorizontalBox::Slot()
+						.FillWidth(0.1f)
+						[
+							SNew(STextBlock)
+							.Text(FText::AsPercent(Item->CountAsPct))
+							.TextStyle(FAppStyle::Get(), "TreeTable.Tooltip")
+						]
+
+						+SHorizontalBox::Slot()
+						.FillWidth(0.9f)
+						[
+							SNew(SProgressBar)
+								.Percent(Item->CountAsPctNormalized)
+								.RefreshRate(0)
+								.BorderPadding(FVector2D(4, 4))
+								.BarFillStyle(EProgressBarFillStyle::Scale)
+								.BackgroundImage(FAppStyle::Get().GetBrush("NoBrush"))
+								.FillImage(FAppStyle::Get().GetBrush("WhiteBrush"))
+						];
+				}
+				else if (Column == Private::ColumnQuantizedReadCount)
+				{
+					return SNew(STextBlock)
+						.Text(FText::AsNumber(Item->Count));
+				}
+				else if (Column == Private::ColumnQuantizedReadTime)
+				{
+					if (Item->Count > 0)
+					{
+						const double AvgDuration = (Item->MinDuration + Item->MaxDuration) / 2.0;
+						return SNew(STextBlock)
+							.Text(FText::FromString(*UE::Insights::FormatTimeAuto(AvgDuration)));
+					}
+					else
+					{
+						return SNew(STextBlock)
+							.Text(LOCTEXT("NotApplicable","N/A"));
+					}
+				}
+
+				return SNullWidget::NullWidget;
+			}
+
+		private:
+			TSharedPtr<Private::FReadSizeHistogramItem> Item;
+		};
+
+
+	}
+
+
+	void SIoStoreAnalysisReadSizeHistogramView::Construct(const FArguments& InArgs)
+	{
+		Super::Construct(
+			Super::FArguments()
+			.ListItemsSource(InArgs._ListItemsSource)
+			.ScrollbarVisibility(EVisibility::Visible)
+			.OnGenerateRow(this, &SIoStoreAnalysisReadSizeHistogramView::OnGenerateRow)
+			.HeaderRow
+			(
+				SNew(SHeaderRow)
+
+				+SHeaderRow::Column(Private::ColumnQuantizedReadSize)
+				.DefaultLabel(LOCTEXT("ColumnQuantizedReadSize", "Quantized Size"))
+				.FillWidth(0.20f)
+
+				+SHeaderRow::Column(Private::ColumnQuantizedReadCountGraph)
+				.DefaultLabel(LOCTEXT("ColumnQuantizedReadGraph", "Histogram"))
+				.FillWidth(0.45f)
+
+				+SHeaderRow::Column(Private::ColumnQuantizedReadCount)
+				.DefaultLabel(LOCTEXT("ColumnQuantizedReadCount", "Num. Reads"))
+				.FillWidth(0.15f)
+
+				+SHeaderRow::Column(Private::ColumnQuantizedReadTime)
+				.DefaultLabel(LOCTEXT("ColumnQuantizedReadTime", "Avg. Duration"))
+				.FillWidth(0.20f)						
+			)
+		);
+	}
+
+
+	TSharedRef<ITableRow> SIoStoreAnalysisReadSizeHistogramView::OnGenerateRow(TSharedPtr<Private::FReadSizeHistogramItem> Item, const TSharedRef<STableViewBase>& OwnerTable)
+	{
+		return SNew(Private::SReadSizeHistogramViewRow, OwnerTable, Item);
+	}
+
+	
+
+	void SIoStoreAnalysisReadSizeHistogramView::SetAnalysisSession( const TraceServices::IAnalysisSession* InAnalysisSession )
+	{
+		AnalysisSession = InAnalysisSession;
+	}
+
+} //namespace UE::IoStoreInsights
+
+#undef LOCTEXT_NAMESPACE
diff --git a/Engine/Plugins/IoStoreInsights/Source/IoStoreInsights/Private/Widgets/SIoStoreAnalysisReadSizeHistogramView.h b/Engine/Plugins/IoStoreInsights/Source/IoStoreInsights/Private/Widgets/SIoStoreAnalysisReadSizeHistogramView.h
new file mode 100644
--- /dev/null
+++ b/Engine/Plugins/IoStoreInsights/Source/IoStoreInsights/Private/Widgets/SIoStoreAnalysisReadSizeHistogramView.h
@@ -0,0 +1,47 @@
+// Copyright Epic Games, Inc. All Rights Reserved.
+
+#pragma once
+
+#include "Widgets/Views/SListView.h"
+
+namespace TraceServices
+{
+	class IAnalysisSession;
+}
+
+namespace UE::IoStoreInsights
+{
+	namespace Private
+	{ 
+		struct FReadSizeHistogramItem
+		{
+			uint64 QuantizedReadSize = 0;
+			uint32 Count = 0;
+			float CountAsPct = 0;
+			float CountAsPctNormalized = 0;
+			double MinDuration = 0;
+			double MaxDuration = 0;
+		};
+	}
+
+	class SIoStoreAnalysisReadSizeHistogramView : public SListView<TSharedPtr<Private::FReadSizeHistogramItem>>
+	{
+		typedef SListView<TSharedPtr<Private::FReadSizeHistogramItem>> Super;
+
+	public:
+		SLATE_BEGIN_ARGS(SIoStoreAnalysisReadSizeHistogramView){}
+			SLATE_ARGUMENT(TArray<TSharedPtr<Private::FReadSizeHistogramItem>>*, ListItemsSource)
+		SLATE_END_ARGS()
+		~SIoStoreAnalysisReadSizeHistogramView() = default;
+
+		void Construct(const FArguments& InArgs);
+		void SetAnalysisSession( const TraceServices::IAnalysisSession* InAnalysisSession );
+
+	private:
+		const TraceServices::IAnalysisSession* AnalysisSession = nullptr;
+
+		TSharedRef<ITableRow> OnGenerateRow(TSharedPtr<Private::FReadSizeHistogramItem> Item, const TSharedRef<STableViewBase>& OwnerTable);
+	};
+
+} //namespace UE::IoStoreInsights
+
diff --git a/Engine/Plugins/IoStoreInsights/Source/IoStoreInsights/Private/Widgets/SIoStoreAnalysisTab.cpp b/Engine/Plugins/IoStoreInsights/Source/IoStoreInsights/Private/Widgets/SIoStoreAnalysisTab.cpp
new file mode 100644
--- /dev/null
+++ b/Engine/Plugins/IoStoreInsights/Source/IoStoreInsights/Private/Widgets/SIoStoreAnalysisTab.cpp
@@ -0,0 +1,399 @@
+// Copyright Epic Games, Inc. All Rights Reserved.
+
+#include "SIoStoreAnalysisTab.h"
+#include "SIoStoreAnalysisReadSizeHistogramView.h"
+#include "SIoStoreAnalysisActivityView.h"
+#include "Model/IoStoreInsightsProvider.h"
+#include "ViewModels/IoStoreInsightsTrack.h"
+#include "Widgets/Text/STextBlock.h"
+#include "Widgets/Views/SHeaderRow.h"
+#include "Widgets/Layout/SWidgetSwitcher.h"
+#include "Widgets/Input/SCheckBox.h"
+#include "InsightsCore/Common/TimeUtils.h"
+#include "Insights/ITimingViewSession.h"
+#include "Insights/ViewModels/ITimingEvent.h"
+#include "IO/IoChunkId.h"
+
+#define LOCTEXT_NAMESPACE "SIoStoreAnalysisTab"
+
+namespace UE::IoStoreInsights
+{
+	void SIoStoreAnalysisTab::Construct(const FArguments& InArgs)
+	{
+		// create main widgets
+		TSharedPtr<SHorizontalBox> ViewModeSelectorsBox;
+		TSharedPtr<SWidgetSwitcher> ViewModeSwitcher;
+		ChildSlot
+		[
+			SNew(SVerticalBox)
+
+			// view mode selection
+			+SVerticalBox::Slot()
+			.AutoHeight()
+			.Padding(8,2)
+			[
+				SAssignNew(ViewModeSelectorsBox, SHorizontalBox)
+			]
+
+			// active view mode panel
+			+SVerticalBox::Slot()
+			.FillHeight(1.0f)
+			.Padding(3)
+			[
+				SAssignNew(ViewModeSwitcher, SWidgetSwitcher)
+				.WidgetIndex(this, &SIoStoreAnalysisTab::GetSelectedViewModeIndex)
+			]
+		];
+
+		// helper to create widgets for a view mode
+		auto AddViewMode = [this, &ViewModeSelectorsBox, &ViewModeSwitcher]( EViewMode InViewMode, const FText& InText)
+		{
+			// add mode selection button
+			ViewModeSelectorsBox->AddSlot()
+			.AutoWidth()
+			.Padding(4)
+			[
+				SNew(SCheckBox)
+				.Style(&FAppStyle::GetWidgetStyle<FCheckBoxStyle>("ToggleButtonCheckBox"))
+				.IsChecked( this, &SIoStoreAnalysisTab::IsViewModeSelected, InViewMode )
+				.OnCheckStateChanged( this, &SIoStoreAnalysisTab::OnViewModeCheckStateChange, InViewMode )
+				[
+					SNew(STextBlock)
+						.Text(InText)
+				]
+			];
+
+			// return a new slot for the mode's panel
+			return ViewModeSwitcher->AddSlot();
+		};
+
+		// add the view modes
+		AddViewMode(EViewMode::ReadSizes, LOCTEXT("ViewMode_ReadSizes", "Read Size Histogram"))
+		[
+			SAssignNew(ReadSizeHistogramView, SIoStoreAnalysisReadSizeHistogramView)
+				.ListItemsSource(&ReadSizeHistogramItems)
+		];
+
+		AddViewMode(EViewMode::ReadActivity, LOCTEXT("ViewMode_ReadActivity", "Read Activity"))
+		[
+			SAssignNew(ActivityInfoTreeView, SIoStoreAnalysisActivityView)
+				.TreeItemsSource(&ActivityInfoArray)
+		];
+
+		// update the data
+		RefreshNodes();
+	}
+
+
+
+	SIoStoreAnalysisTab::~SIoStoreAnalysisTab()
+	{
+		if (TimingViewSession)
+		{
+			TimingViewSession->OnTimeMarkerChanged().RemoveAll(this);
+			TimingViewSession->OnSelectionChanged().RemoveAll(this);
+		}
+	}
+
+
+
+	void SIoStoreAnalysisTab::SetSession(UE::Insights::Timing::ITimingViewSession* InTimingViewSession, const TraceServices::IAnalysisSession* InAnalysisSession, const FIoStoreInsightsViewSharedState* InSharedStatePtr)
+	{
+		if (TimingViewSession == InTimingViewSession)
+		{
+			return;
+		}
+
+		if (TimingViewSession)
+		{
+			TimingViewSession->OnTimeMarkerChanged().RemoveAll(this);
+			TimingViewSession->OnSelectionChanged().RemoveAll(this);
+			TimingViewSession->OnSelectedEventChanged().RemoveAll(this);
+		}
+
+		TimingViewSession = InTimingViewSession;
+		AnalysisSession = InAnalysisSession;
+		SharedStatePtr = InSharedStatePtr;
+		ReadSizeHistogramView->SetAnalysisSession(AnalysisSession);
+		ActivityInfoTreeView->SetAnalysisSession(AnalysisSession);
+
+		if (InTimingViewSession)
+		{
+			InTimingViewSession->OnTimeMarkerChanged().AddSP(this, &SIoStoreAnalysisTab::HandleTimeMarkerChanged);
+			InTimingViewSession->OnSelectionChanged().AddSP(this, &SIoStoreAnalysisTab::HandleSelectionChanged);
+			TimingViewSession->OnSelectedEventChanged().AddSP(this, &SIoStoreAnalysisTab::HandleSelectionEventChanged);
+		}
+
+		RefreshNodes();
+	}
+
+
+
+	bool SIoStoreAnalysisTab::IsSessionSet() const
+	{
+		return TimingViewSession != nullptr;
+	}
+
+
+
+	void SIoStoreAnalysisTab::HandleTimeMarkerChanged(UE::Insights::Timing::ETimeChangedFlags InFlags, double InTimeMarker)
+	{
+		if (!FMath::IsNearlyEqual(StartTime, InTimeMarker) && !FMath::IsNearlyEqual(EndTime, InTimeMarker))
+		{
+			StartTime = InTimeMarker;
+			EndTime = InTimeMarker;
+
+			RefreshNodes();
+		}
+	}
+
+
+
+	void SIoStoreAnalysisTab::HandleSelectionChanged(UE::Insights::Timing::ETimeChangedFlags InFlags, double InStartTime, double InEndTime)
+	{
+		if(InFlags != UE::Insights::Timing::ETimeChangedFlags::Interactive)
+		{
+			StartTime = InStartTime;
+			EndTime = InEndTime;
+
+			RefreshNodes();
+		}
+	}
+
+
+
+	void SIoStoreAnalysisTab::HandleSelectionEventChanged(const TSharedPtr<const ITimingEvent> InEvent)
+	{
+		if (InEvent)
+		{
+			const double EventStartTime = InEvent->GetStartTime();
+			const double EventEndTime = InEvent->GetEndTime();
+			if (!FMath::IsNearlyEqual(StartTime, EventStartTime) && !FMath::IsNearlyEqual(EndTime, EventEndTime))
+			{
+				StartTime = EventStartTime;
+				EndTime = EventEndTime;
+				RefreshNodes();
+			}
+		}
+	}
+
+
+
+	int32 SIoStoreAnalysisTab::GetSelectedViewModeIndex() const
+	{
+		return (int32)ViewMode;
+	}
+
+
+
+	ECheckBoxState SIoStoreAnalysisTab::IsViewModeSelected(EViewMode Mode) const
+	{
+		return (Mode == ViewMode) ? ECheckBoxState::Checked : ECheckBoxState::Unchecked;
+	}
+
+
+
+	void SIoStoreAnalysisTab::OnViewModeCheckStateChange(ECheckBoxState State, EViewMode Mode)
+	{
+		if (State == ECheckBoxState::Checked)
+		{
+			ViewMode = Mode;
+		}
+	}
+
+
+
+	void SIoStoreAnalysisTab::RefreshNodes()
+	{
+		if (TimingViewSession && AnalysisSession && SharedStatePtr)
+		{
+			if (EndTime < 0)
+			{
+				EndTime = +std::numeric_limits<double>::infinity();
+			}
+
+			if (StartTime <= EndTime && EndTime >= 0.0)
+			{
+				const FIoStoreInsightsProvider* Provider = AnalysisSession->ReadProvider<FIoStoreInsightsProvider>(IIoStoreInsightsProvider::ProviderName);
+				if (Provider)
+				{
+					TraceServices::FAnalysisSessionReadScope SessionReadScope(*AnalysisSession);
+					RefreshNodes_IoStoreActivity(Provider);
+				}
+			}
+		}
+
+
+		ReadSizeHistogramView->RebuildList();
+
+		ActivityInfoTreeView->ApplySort();
+		ActivityInfoTreeView->RebuildList();
+	}
+
+
+
+	void SIoStoreAnalysisTab::RefreshNodes_IoStoreActivity(const FIoStoreInsightsProvider* Provider)
+	{
+		const IIoStoreInsightsProvider* IoStoreActivityProvider = AnalysisSession ? AnalysisSession->ReadProvider<IIoStoreInsightsProvider>(FIoStoreInsightsProvider::ProviderName) : nullptr;
+
+		uint32 TotalReads = 0;
+
+		uint64 MinQuantizedReadSize = std::numeric_limits<uint64>::max();
+		uint64 MaxQuantizedReadSize = 0;
+		TMap<uint64,TSharedPtr<Private::FReadSizeHistogramItem>> QuantizedSizeToHisogramItemMap;
+		ReadSizeHistogramItems.Reset();
+
+		TMap<uint32,TArray<const FIoStoreActivity*>> DataSourceToActivityMap;
+		ActivityInfoArray.Reset();
+
+
+		Provider->EnumerateIoStoreRequests([this, &DataSourceToActivityMap, &TotalReads, &MinQuantizedReadSize, &MaxQuantizedReadSize, &QuantizedSizeToHisogramItemMap](const FIoStoreRequest& IoStoreRequest, const IIoStoreInsightsProvider::Timeline& Timeline)
+		{
+			Timeline.EnumerateEvents(StartTime, EndTime, [this, &DataSourceToActivityMap, &TotalReads, &MinQuantizedReadSize, &MaxQuantizedReadSize, &QuantizedSizeToHisogramItemMap](double EventStartTime, double EventEndTime, uint32 EventDepth, const FIoStoreActivity* IoStoreActivity)
+			{
+				if (IoStoreActivity->ActivityType == EIoStoreActivityType::Request_Read && IoStoreActivity->EndTime > 0 && !IoStoreActivity->Failed)
+				{
+					// build the table of quantized reads
+					uint64 QuantizedReadSize = FPlatformMath::RoundUpToPowerOfTwo64(IoStoreActivity->ActualSize);
+					MaxQuantizedReadSize = FMath::Max(MaxQuantizedReadSize, QuantizedReadSize);
+					MinQuantizedReadSize = FMath::Min(MinQuantizedReadSize, QuantizedReadSize);
+
+					double Duration = (IoStoreActivity->EndTime - IoStoreActivity->StartTime);
+
+					TSharedPtr<Private::FReadSizeHistogramItem> HistogramItem;
+					if (!QuantizedSizeToHisogramItemMap.Contains(QuantizedReadSize))
+					{
+						HistogramItem = MakeShared<Private::FReadSizeHistogramItem>();
+						HistogramItem->QuantizedReadSize = QuantizedReadSize;
+						HistogramItem->MinDuration = Duration;
+						HistogramItem->MaxDuration = Duration;
+						QuantizedSizeToHisogramItemMap.Add(QuantizedReadSize, HistogramItem);
+						ReadSizeHistogramItems.Add(HistogramItem);
+					}
+					else
+					{
+						HistogramItem = QuantizedSizeToHisogramItemMap.FindRef(QuantizedReadSize);
+					}
+					
+					TotalReads++;
+					HistogramItem->Count++;
+					HistogramItem->MinDuration = FMath::Min(Duration, HistogramItem->MinDuration);
+					HistogramItem->MaxDuration = FMath::Max(Duration, HistogramItem->MaxDuration);
+
+					// collect the activities associated with the package or other data source
+					uint32 SourceHash;
+					if (*IoStoreActivity->IoStoreRequest->PackageName)
+					{
+						SourceHash = PointerHash(IoStoreActivity->IoStoreRequest->PackageName);
+					}
+					else if (*IoStoreActivity->IoStoreRequest->ExtraTag)
+					{
+						SourceHash = PointerHash(IoStoreActivity->IoStoreRequest->ExtraTag);
+					}
+					else
+					{
+						SourceHash = GetTypeHash(IoStoreActivity->IoStoreRequest->PackageId);
+					}
+					DataSourceToActivityMap.FindOrAdd(SourceHash).Add(IoStoreActivity);
+				}
+
+				return TraceServices::EEventEnumerate::Continue;
+			});
+
+			return true;
+		});
+
+
+		// fill in any missing items
+		for (uint64 MissingQuantizedReadSize = MinQuantizedReadSize; MissingQuantizedReadSize < MaxQuantizedReadSize || MissingQuantizedReadSize == 0; MissingQuantizedReadSize <<= 1)
+		{
+			if (!QuantizedSizeToHisogramItemMap.Contains(MissingQuantizedReadSize))
+			{
+				TSharedPtr<Private::FReadSizeHistogramItem> HistogramItem = MakeShared<Private::FReadSizeHistogramItem>();
+				HistogramItem->QuantizedReadSize = MissingQuantizedReadSize;
+				QuantizedSizeToHisogramItemMap.Add(MissingQuantizedReadSize, HistogramItem);
+				ReadSizeHistogramItems.Add(HistogramItem);
+			}
+		}
+
+		// sort by read size
+		ReadSizeHistogramItems.Sort( []( const TSharedPtr<Private::FReadSizeHistogramItem>& A, const TSharedPtr<Private::FReadSizeHistogramItem>& B)
+		{
+			return A->QuantizedReadSize < B->QuantizedReadSize;
+		});
+
+		// compute percentages
+		float MaxPct = 0.0f;
+		for (const TSharedPtr<Private::FReadSizeHistogramItem>& HistogramItem : ReadSizeHistogramItems)
+		{
+			HistogramItem->CountAsPct = (float)HistogramItem->Count / (float)TotalReads;
+			MaxPct = FMath::Max(MaxPct, HistogramItem->CountAsPct);
+		}
+		for (const TSharedPtr<Private::FReadSizeHistogramItem>& HistogramItem : ReadSizeHistogramItems)
+		{
+			HistogramItem->CountAsPctNormalized = HistogramItem->CountAsPct * (1.0f / MaxPct);
+		}
+
+		// build the tree
+		for (const auto& Itr : DataSourceToActivityMap)
+		{
+			const TArray<const FIoStoreActivity*>& Activities = Itr.Value;
+			if (Activities.Num() == 0)
+			{
+				continue;
+			}
+			const FIoStoreActivity* FirstActivity = Activities[0];
+
+			// prepare parent information, based on the first activity
+			FStringBuilderBase PackageNameAndTag;
+			if (*FirstActivity->IoStoreRequest->PackageName)
+			{
+				PackageNameAndTag.Append(FirstActivity->IoStoreRequest->PackageName);
+			}
+			else if (IoStoreActivityProvider != nullptr)
+			{
+				PackageNameAndTag.Append(IoStoreActivityProvider->GetPackageName(FirstActivity->IoStoreRequest->PackageId));
+			}
+			else
+			{
+				PackageNameAndTag.Appendf(TEXT("%x"), FirstActivity->IoStoreRequest->PackageId);
+			}
+			if (*FirstActivity->IoStoreRequest->ExtraTag)
+			{
+				PackageNameAndTag.Append(TEXT(" Tag: "));
+				PackageNameAndTag.Append(FirstActivity->IoStoreRequest->ExtraTag);
+			}
+			if (Itr.Value.Num() > 1)
+			{
+				PackageNameAndTag.Appendf( TEXT(" (%d)"), Itr.Value.Num());
+			}
+
+			TSharedPtr<Private::FActivityInfoParent> Parent = MakeShared<Private::FActivityInfoParent>();
+			Parent->DisplayName = PackageNameAndTag.ToString();
+			Parent->IoStoreActivity = FirstActivity;
+			Parent->bIsParent = true;
+			ActivityInfoArray.Add(Parent);
+
+			// process all activities
+			for (const FIoStoreActivity* Activity : Activities)
+			{
+				double Duration = (Activity->EndTime > 0) ? (Activity->EndTime - Activity->StartTime) : +std::numeric_limits<double>::infinity();
+				Parent->Duration += Duration;
+				Parent->TotalSize += Activity->ActualSize;
+
+				// only add child items if there's more than one - we'll display the information in the parent if there's only one
+				if (Itr.Value.Num() > 1)
+				{
+					TSharedPtr<Private::FActivityInfoChild> Child = MakeShared<Private::FActivityInfoChild>();
+					Child->IoStoreActivity = Activity;
+					Child->Duration = Duration;
+					Child->bIsParent = false;
+					Parent->Children.Add(Child);
+				}
+			}
+		}
+	}
+
+} //namespace UE::IoStoreInsights
+
+#undef LOCTEXT_NAMESPACE
diff --git a/Engine/Plugins/IoStoreInsights/Source/IoStoreInsights/Private/Widgets/SIoStoreAnalysisTab.h b/Engine/Plugins/IoStoreInsights/Source/IoStoreInsights/Private/Widgets/SIoStoreAnalysisTab.h
new file mode 100644
--- /dev/null
+++ b/Engine/Plugins/IoStoreInsights/Source/IoStoreInsights/Private/Widgets/SIoStoreAnalysisTab.h
@@ -0,0 +1,77 @@
+// Copyright Epic Games, Inc. All Rights Reserved.
+
+#pragma once
+
+#include "Widgets/SCompoundWidget.h"
+
+class ITimingEvent;
+enum class ECheckBoxState : uint8;
+
+namespace UE::Insights::Timing
+{
+	class ITimingViewSession;
+	enum class ETimeChangedFlags;
+}
+
+namespace TraceServices
+{
+	class IAnalysisSession;
+}
+
+namespace UE::IoStoreInsights
+{
+	class SIoStoreAnalysisActivityView;
+	class SIoStoreAnalysisReadSizeHistogramView;
+	class FIoStoreInsightsProvider;
+	class FIoStoreInsightsViewSharedState;
+	namespace Private { struct FReadSizeHistogramItem; }
+	namespace Private { struct FActivityInfo; }
+
+	class SIoStoreAnalysisTab : public SCompoundWidget
+	{
+	public:
+		SLATE_BEGIN_ARGS(SIoStoreAnalysisTab) {}
+		SLATE_END_ARGS()
+		~SIoStoreAnalysisTab();
+
+		void Construct(const FArguments& InArgs);
+
+		void SetSession(UE::Insights::Timing::ITimingViewSession* InTimingViewSession, const TraceServices::IAnalysisSession* InAnalysisSession, const FIoStoreInsightsViewSharedState* InSharedStatePtr);
+		bool IsSessionSet() const;
+
+	private:
+		enum class EViewMode : uint8
+		{
+			ReadSizes,
+			ReadActivity,
+		};
+
+		void HandleTimeMarkerChanged(UE::Insights::Timing::ETimeChangedFlags InFlags, double InTimeMarker);
+		void HandleSelectionChanged(UE::Insights::Timing::ETimeChangedFlags InFlags, double StartTime, double EndTime);
+		void HandleSelectionEventChanged(const TSharedPtr<const ITimingEvent> InEvent);
+
+		int32 GetSelectedViewModeIndex() const;
+		ECheckBoxState IsViewModeSelected(EViewMode Mode) const;
+		void OnViewModeCheckStateChange(ECheckBoxState State, EViewMode Mode);
+
+		void RefreshNodes();
+		void RefreshNodes_IoStoreActivity(const FIoStoreInsightsProvider* Provider);
+
+	private:
+		UE::Insights::Timing::ITimingViewSession* TimingViewSession = nullptr;
+		const TraceServices::IAnalysisSession* AnalysisSession = nullptr;
+		const FIoStoreInsightsViewSharedState* SharedStatePtr = nullptr;
+
+		double StartTime = -1;
+		double EndTime = -1;
+		EViewMode ViewMode = EViewMode::ReadSizes;
+
+		TArray<TSharedPtr<Private::FReadSizeHistogramItem>> ReadSizeHistogramItems;
+		TSharedPtr<SIoStoreAnalysisReadSizeHistogramView> ReadSizeHistogramView;
+
+		TArray<TSharedPtr<Private::FActivityInfo>> ActivityInfoArray;
+		TSharedPtr<SIoStoreAnalysisActivityView> ActivityInfoTreeView;
+	};
+
+} //namespace UE::IoStoreInsights
+
