Replace bespoke activity table view with an Insights TableTreeView.

 - Also, Insights Provider now fixes up deferred package names automatically

#rnx

diff --git a/Engine/Plugins/IoStoreInsights/Source/IoStoreInsights/Private/Model/IoStoreInsightsProvider.cpp b/Engine/Plugins/IoStoreInsights/Source/IoStoreInsights/Private/Model/IoStoreInsightsProvider.cpp
--- a/Engine/Plugins/IoStoreInsights/Source/IoStoreInsights/Private/Model/IoStoreInsightsProvider.cpp
+++ b/Engine/Plugins/IoStoreInsights/Source/IoStoreInsights/Private/Model/IoStoreInsightsProvider.cpp
@@ -45,6 +45,14 @@ namespace UE::IoStoreInsights
 		IoStoreRequestInfo.IoStoreRequestInfo.PackageName = Session.StoreString(PackageName);
 		IoStoreRequestInfo.IoStoreRequestInfo.ExtraTag = Session.StoreString(ExtraTag);
 		IoStoreRequestInfo.ActivityTimeline = MakeShared<TimelineInternal>(Session.GetLinearAllocator());
+
+		if (*PackageName == '\0' && PackageId != 0)
+		{
+			TArray<uint32>& PendingPackageRequestIndices = PendingPackageNameMap.FindOrAdd(PackageId);
+			PendingPackageRequestIndices.Add(IoStoreRequestIndex);
+		}
+			
+
 		return IoStoreRequestIndex;
 	}
 
@@ -90,6 +98,20 @@ namespace UE::IoStoreInsights
 	void FIoStoreInsightsProvider::AddPackageMapping(uint64 PackageId, const TCHAR* PackageName)
 	{
 		PackageMap.Add(PackageId, PackageName);
+
+		// see if there are any requests that referred to this package & update the package name now
+		const TArray<uint32>* PendingPackageRequestIndicesPtr = PendingPackageNameMap.Find(PackageId);
+		if (PendingPackageRequestIndicesPtr != nullptr && *PackageName != '\0')
+		{
+			const TCHAR* StoredPackageName = Session.StoreString(PackageName);
+			for (uint32 IoStoreRequestIndex : *PendingPackageRequestIndicesPtr)
+			{
+				FIoStoreRequestInfoInternal& IoStoreRequestInfoInternal = IoStoreRequests[IoStoreRequestIndex];
+				IoStoreRequestInfoInternal.IoStoreRequestInfo.PackageName = StoredPackageName;
+			}
+
+			PendingPackageNameMap.Remove(PackageId);
+		}
 	}
 
 
diff --git a/Engine/Plugins/IoStoreInsights/Source/IoStoreInsights/Private/Model/IoStoreInsightsProvider.h b/Engine/Plugins/IoStoreInsights/Source/IoStoreInsights/Private/Model/IoStoreInsightsProvider.h
--- a/Engine/Plugins/IoStoreInsights/Source/IoStoreInsights/Private/Model/IoStoreInsightsProvider.h
+++ b/Engine/Plugins/IoStoreInsights/Source/IoStoreInsights/Private/Model/IoStoreInsightsProvider.h
@@ -55,6 +55,8 @@ namespace UE::IoStoreInsights
 		TraceServices::TPagedArray<FIoStoreActivity> IoStoreRequestStates;
 		TMap<uint64, const TCHAR*> PackageMap;
 		TMap<uint64, const TCHAR*> BackendNameMap;
+
+		TMap<uint64,TArray<uint32>> PendingPackageNameMap;
 	};
 
 } // namespace UE::IoStoreInsights
diff --git a/Engine/Plugins/IoStoreInsights/Source/IoStoreInsights/Private/ViewModels/IoStoreActivityTable.cpp b/Engine/Plugins/IoStoreInsights/Source/IoStoreInsights/Private/ViewModels/IoStoreActivityTable.cpp
new file mode 100644
--- /dev/null
+++ b/Engine/Plugins/IoStoreInsights/Source/IoStoreInsights/Private/ViewModels/IoStoreActivityTable.cpp
@@ -0,0 +1,239 @@
+// Copyright Epic Games, Inc. All Rights Reserved.
+
+#include "IoStoreActivityTable.h"
+#include "IoStoreActivityTableTreeNode.h"
+#include "InsightsCore/Table/ViewModels/TableCellValue.h"
+#include "InsightsCore/Table/ViewModels/TableCellValueFormatter.h"
+#include "InsightsCore/Table/ViewModels/TableCellValueGetter.h"
+#include "InsightsCore/Table/ViewModels/TableCellValueSorter.h"
+#include "InsightsCore/Table/ViewModels/TableColumn.h"
+#include "IIoStoreInsightsProvider.h"
+#include "IO/IoChunkId.h"
+
+#define LOCTEXT_NAMESPACE "UE::IoStoreInsights::FIoStoreActivityTable"
+
+namespace UE::IoStoreInsights
+{
+	using namespace UE::Insights;
+
+	const FName FActivityTableColumns::ColumnRequestPackage("Package");
+	const FName FActivityTableColumns::ColumnRequestOffset("Offset");
+	const FName FActivityTableColumns::ColumnRequestSize("Size");
+	const FName FActivityTableColumns::ColumnRequestDuration("Duration");
+	const FName FActivityTableColumns::ColumnRequestChunkId("ChunkId");
+	const FName FActivityTableColumns::ColumnRequestChunkType("ChunkType");
+	const FName FActivityTableColumns::ColumnRequestStartTime("StartTime");
+
+	typedef FTableCellValue (*FActivityFieldGetterFn) (const FTableColumn&, const FIoStoreActivity*);
+
+	template<FActivityFieldGetterFn Getter>
+	class FActivityColumnValueGetter : public FTableCellValueGetter
+	{
+	public:
+		virtual const TOptional<FTableCellValue> GetValue(const FTableColumn& Column, const FBaseTreeNode& Node) const override
+		{
+			if (Node.IsGroup())
+			{
+				const FTableTreeNode& NodePtr = static_cast<const FTableTreeNode&>(Node);
+				if (NodePtr.HasAggregatedValue(Column.GetId()))
+				{
+					return NodePtr.GetAggregatedValue(Column.GetId());
+				}
+			}
+			else //if (Node->Is<FIoStoreActivityNode>())
+			{
+				const FIoStoreActivityNode& ActivityNode = static_cast<const FIoStoreActivityNode&>(Node);
+				const FIoStoreActivity* Activity = ActivityNode.GetActivity();
+				if (Activity)
+				{
+					return Getter(Column, Activity);
+				}
+			}
+
+			return TOptional<FTableCellValue>();
+		}
+	};
+
+	const TCHAR* GetActivityDisplayName(const FIoStoreActivity* Activity)
+	{
+		if (*Activity->IoStoreRequest->PackageName)
+		{
+			return Activity->IoStoreRequest->PackageName;
+		}
+		else if (*Activity->IoStoreRequest->ExtraTag)
+		{
+			return Activity->IoStoreRequest->ExtraTag;
+		}
+		else
+		{
+			return TEXT("(Unknown Package)");
+		}
+	}
+
+
+	struct DefaultActivityFieldGetterFuncs
+	{
+		static FTableCellValue GetPackage  (const FTableColumn& Column, const FIoStoreActivity* Activity) { return FTableCellValue(GetActivityDisplayName(Activity)); }
+		static FTableCellValue GetOffset   (const FTableColumn& Column, const FIoStoreActivity* Activity) { return FTableCellValue((int64)Activity->IoStoreRequest->Offset); }
+		static FTableCellValue GetSize     (const FTableColumn& Column, const FIoStoreActivity* Activity) { return FTableCellValue((int64)Activity->ActualSize); }
+		static FTableCellValue GetDuration (const FTableColumn& Column, const FIoStoreActivity* Activity) { return FTableCellValue(Activity->EndTime - Activity->StartTime); }
+		static FTableCellValue GetChunkId  (const FTableColumn& Column, const FIoStoreActivity* Activity) { return FTableCellValue((int64)Activity->IoStoreRequest->ChunkIdHash); }
+		static FTableCellValue GetChunkType(const FTableColumn& Column, const FIoStoreActivity* Activity) { return FTableCellValue(FText::FromString(LexToString((EIoChunkType)Activity->IoStoreRequest->ChunkType))); }
+		static FTableCellValue GetStartTime(const FTableColumn& Column, const FIoStoreActivity* Activity) { return FTableCellValue(Activity->StartTime); }
+	};
+
+
+
+	void FIoStoreActivityTable::Reset()
+	{
+		FTable::Reset();
+		AddDefaultColumns();
+	}
+
+	void FIoStoreActivityTable::AddDefaultColumns()
+	{
+		using namespace UE::Insights;
+
+		// Hierarchy Column (special case)
+		{
+			AddHierarchyColumn(-1, nullptr);
+			const TSharedRef<FTableColumn>& ColumnRef = GetColumns()[0];
+			ColumnRef->SetInitialWidth(75.0f);
+			ColumnRef->SetShortName(LOCTEXT("ActivityColumnName",   "Hierarchy"));
+			ColumnRef->SetTitleName(LOCTEXT("ActivityColumnTitle",  "All Activities"));
+			ColumnRef->SetDescription(LOCTEXT("ActivityColumnDesc", "Hierarchy of all activities"));
+		}
+
+		// helper to add common column data
+		int32 ColumnIndex = 0;
+		auto MakeColumn = [this, &ColumnIndex]( FName InColumnName, const FText& InShortName, const FText& InTitleName, const FText& InDescription, float InWidth )
+		{
+			TSharedRef<FTableColumn> ColumnRef = MakeShared<FTableColumn>(InColumnName);
+			AddColumn(ColumnRef);
+
+			FTableColumn& Column = *ColumnRef;
+			Column.SetIndex(ColumnIndex++);
+			Column.SetShortName(InShortName);
+			Column.SetTitleName(InTitleName);
+			Column.SetDescription(InDescription);
+			Column.SetHorizontalAlignment(HAlign_Left);
+			Column.SetInitialWidth(InWidth);
+			return ColumnRef;
+		};
+
+		// Package Column
+		{
+			TSharedRef<FTableColumn> ColumnRef = MakeColumn(FActivityTableColumns::ColumnRequestPackage,
+															LOCTEXT("PackageColumnName",  "Package"),
+															LOCTEXT("PackageColumnTitle", "Package"),
+															LOCTEXT("PackageColumnDesc",  "Package Name/Tag being read"),
+															300.0f);
+			FTableColumn& Column = *ColumnRef;
+			Column.SetFlags(ETableColumnFlags::ShouldBeVisible | ETableColumnFlags::CanBeHidden | ETableColumnFlags::CanBeFiltered);
+			Column.SetDataType(ETableCellDataType::CString);
+			Column.SetValueGetter(MakeShared<FActivityColumnValueGetter<DefaultActivityFieldGetterFuncs::GetPackage>>());
+			Column.SetValueFormatter(MakeShared<FCStringValueFormatterAsText>());
+			Column.SetValueSorter(MakeShared<FSorterByCStringValue>(ColumnRef));
+			Column.SetAggregation(ETableColumnAggregation::SameValue);
+		}
+
+		// Offset Column
+		{
+			TSharedRef<FTableColumn> ColumnRef = MakeColumn(FActivityTableColumns::ColumnRequestOffset,
+															LOCTEXT("OffsetColumnName",  "Offset"),
+															LOCTEXT("OffsetColumnTitle", "Read Offset"),
+															LOCTEXT("OffsetColumnDesc",  "Offset into chunk that was requested"),
+															50.0f);
+			FTableColumn& Column = *ColumnRef;
+			Column.SetFlags(ETableColumnFlags::ShouldBeVisible | ETableColumnFlags::CanBeHidden | ETableColumnFlags::CanBeFiltered);
+			Column.SetDataType(ETableCellDataType::Int64);
+			Column.SetValueGetter(MakeShared<FActivityColumnValueGetter<DefaultActivityFieldGetterFuncs::GetOffset>>());
+			Column.SetValueFormatter(MakeShared<FInt64ValueFormatterAsNumber>());
+			Column.SetValueSorter(MakeShared<FSorterByInt64Value>(ColumnRef));
+			Column.SetAggregation(ETableColumnAggregation::None);
+		}
+
+		// Size Column
+		{
+			TSharedRef<FTableColumn> ColumnRef = MakeColumn(FActivityTableColumns::ColumnRequestSize,
+															LOCTEXT("SizeColumnName",  "Size"),
+															LOCTEXT("SizeColumnTitle", "Read Size"),
+															LOCTEXT("SizeColumnDesc",  "Size of the data that was returned by the IoDispatcher"),
+															50.0f);
+			FTableColumn& Column = *ColumnRef;
+			Column.SetFlags(ETableColumnFlags::ShouldBeVisible | ETableColumnFlags::CanBeHidden | ETableColumnFlags::CanBeFiltered);
+			Column.SetDataType(ETableCellDataType::Int64);
+			Column.SetValueGetter(MakeShared<FActivityColumnValueGetter<DefaultActivityFieldGetterFuncs::GetSize>>());
+			Column.SetValueFormatter(MakeShared<FInt64ValueFormatterAsMemory>());
+			Column.SetValueSorter(MakeShared<FSorterByInt64Value>(ColumnRef));
+			Column.SetAggregation(ETableColumnAggregation::Sum);
+		}
+
+		// Duration Column
+		{
+			TSharedRef<FTableColumn> ColumnRef = MakeColumn(FActivityTableColumns::ColumnRequestDuration,
+															LOCTEXT("DurationColumnName",  "Duration"),
+															LOCTEXT("DurationColumnTitle", "Request Duration"),
+															LOCTEXT("DurationColumnDesc",  "How long the request took to complete"),
+															50.0f);
+			FTableColumn& Column = *ColumnRef;
+			Column.SetFlags(ETableColumnFlags::ShouldBeVisible | ETableColumnFlags::CanBeHidden | ETableColumnFlags::CanBeFiltered);
+			Column.SetDataType(ETableCellDataType::Double);
+			Column.SetValueGetter(MakeShared<FActivityColumnValueGetter<DefaultActivityFieldGetterFuncs::GetDuration>>());
+			Column.SetValueFormatter(MakeShared<FDoubleValueFormatterAsTimeAuto>());
+			Column.SetValueSorter(MakeShared<FSorterByDoubleValue>(ColumnRef));
+			Column.SetAggregation(ETableColumnAggregation::Sum);
+		}
+
+		// ChunkId Column
+		{
+			TSharedRef<FTableColumn> ColumnRef = MakeColumn(FActivityTableColumns::ColumnRequestChunkId,
+															LOCTEXT("ChunkIdColumnName",  "ChunkId"),
+															LOCTEXT("ChunkIdColumnTitle", "ChunkId Hash"),
+															LOCTEXT("ChunkIdColumnDesc",  "Hash of the Chunk Id"),
+															50.0f);
+			FTableColumn& Column = *ColumnRef;
+			Column.SetFlags(ETableColumnFlags::CanBeHidden | ETableColumnFlags::CanBeFiltered);
+			Column.SetDataType(ETableCellDataType::Int64);
+			Column.SetValueGetter(MakeShared<FActivityColumnValueGetter<DefaultActivityFieldGetterFuncs::GetChunkId>>());
+			Column.SetValueFormatter(MakeShared<FInt64ValueFormatterAsHex32>());
+			Column.SetValueSorter(MakeShared<FSorterByInt64Value>(ColumnRef));
+			Column.SetAggregation(ETableColumnAggregation::SameValue); // NB. non-CString SameValue aggregation currently not supported
+		}
+
+		// Chunk Type Column
+		{
+			TSharedRef<FTableColumn> ColumnRef = MakeColumn(FActivityTableColumns::ColumnRequestChunkType,
+															LOCTEXT("ChunkTypeName",  "Chunk Type"),
+															LOCTEXT("ChunkTypeTitle", "Chunk Type"),
+															LOCTEXT("ChunkTypeDesc",  "The type of chunk that was requested"),
+															50.0f);
+			FTableColumn& Column = *ColumnRef;
+			Column.SetFlags(ETableColumnFlags::ShouldBeVisible | ETableColumnFlags::CanBeHidden | ETableColumnFlags::CanBeFiltered);
+			Column.SetDataType(ETableCellDataType::Text);
+			Column.SetValueGetter(MakeShared<FActivityColumnValueGetter<DefaultActivityFieldGetterFuncs::GetChunkType>>());
+			Column.SetValueFormatter(MakeShared<FTextValueFormatter>()); // NB. using FText not CString because EIoChunkType LexToString returns FString not TCHAR*
+			Column.SetValueSorter(MakeShared<FSorterByTextValue>(ColumnRef));
+			Column.SetAggregation(ETableColumnAggregation::SameValue);
+		}
+
+		// Start Time Column
+		{
+			TSharedRef<FTableColumn> ColumnRef = MakeColumn(FActivityTableColumns::ColumnRequestStartTime,
+															LOCTEXT("StartTimeName",  "Start Time"),
+															LOCTEXT("StartTimeTitle", "Start Time"),
+															LOCTEXT("StartTimeDesc",  "Time the request was started"),
+															50.0f);
+			FTableColumn& Column = *ColumnRef;
+			Column.SetFlags(ETableColumnFlags::ShouldBeVisible | ETableColumnFlags::CanBeHidden | ETableColumnFlags::CanBeFiltered);
+			Column.SetDataType(ETableCellDataType::Double);
+			Column.SetValueGetter(MakeShared<FActivityColumnValueGetter<DefaultActivityFieldGetterFuncs::GetStartTime>>());
+			Column.SetValueFormatter(MakeShared<FDoubleValueFormatterAsTimeAuto>());
+			Column.SetValueSorter(MakeShared<FSorterByDoubleValue>(ColumnRef));
+			Column.SetAggregation(ETableColumnAggregation::Min);
+		}
+	}
+
+} //namespace UE::IoStoreInsights
+
+#undef LOCTEXT_NAMESPACE
diff --git a/Engine/Plugins/IoStoreInsights/Source/IoStoreInsights/Private/ViewModels/IoStoreActivityTable.h b/Engine/Plugins/IoStoreInsights/Source/IoStoreInsights/Private/ViewModels/IoStoreActivityTable.h
new file mode 100644
--- /dev/null
+++ b/Engine/Plugins/IoStoreInsights/Source/IoStoreInsights/Private/ViewModels/IoStoreActivityTable.h
@@ -0,0 +1,48 @@
+// Copyright Epic Games, Inc. All Rights Reserved.
+
+#pragma once
+
+#include "CoreTypes.h"
+#include "Containers/Array.h"
+#include "UObject/NameTypes.h"
+#include "InsightsCore/Table/ViewModels/Table.h"
+
+namespace UE::Insights { class FTableColumn; }
+
+namespace UE::IoStoreInsights
+{
+	struct FIoStoreActivity;
+
+	struct FActivityTableColumns
+	{
+		static const FName ColumnRequestPackage;
+		static const FName ColumnRequestOffset;
+		static const FName ColumnRequestSize;
+		static const FName ColumnRequestDuration;
+		static const FName ColumnRequestChunkId;
+		static const FName ColumnRequestChunkType;
+		static const FName ColumnRequestStartTime;
+	};
+
+	class FIoStoreActivityTable : public UE::Insights::FTable
+	{
+	public:
+		virtual ~FIoStoreActivityTable() = default;
+
+		virtual void Reset() override;
+
+		TArray<const FIoStoreActivity*>& GetActivities() { return Activities; }
+		const TArray<const FIoStoreActivity*>& GetActivities() const { return Activities; }
+
+		bool IsValidRowIndex(int32 InIndex) const { return Activities.IsValidIndex(InIndex); }
+		const FIoStoreActivity* GetActivity(int32 InIndex) const { return IsValidRowIndex(InIndex) ? Activities[InIndex] : nullptr; }
+		const FIoStoreActivity* GetActivityChecked(int32 InIndex) const { check(IsValidRowIndex(InIndex)); return Activities[InIndex]; }
+
+	private:
+		void AddDefaultColumns();
+
+		TArray<const FIoStoreActivity*> Activities;
+	};
+
+
+} //namespace UE::IoStoreInsights
diff --git a/Engine/Plugins/IoStoreInsights/Source/IoStoreInsights/Private/ViewModels/IoStoreActivityTableTreeNode.cpp b/Engine/Plugins/IoStoreInsights/Source/IoStoreInsights/Private/ViewModels/IoStoreActivityTableTreeNode.cpp
new file mode 100644
--- /dev/null
+++ b/Engine/Plugins/IoStoreInsights/Source/IoStoreInsights/Private/ViewModels/IoStoreActivityTableTreeNode.cpp
@@ -0,0 +1,9 @@
+// Copyright Epic Games, Inc. All Rights Reserved.
+
+#include "IoStoreActivityTableTreeNode.h"
+
+namespace UE::IoStoreInsights
+{
+	INSIGHTS_IMPLEMENT_RTTI(FIoStoreActivityNode)
+
+} // UE::IoStoreInsights
diff --git a/Engine/Plugins/IoStoreInsights/Source/IoStoreInsights/Private/ViewModels/IoStoreActivityTableTreeNode.h b/Engine/Plugins/IoStoreInsights/Source/IoStoreInsights/Private/ViewModels/IoStoreActivityTableTreeNode.h
new file mode 100644
--- /dev/null
+++ b/Engine/Plugins/IoStoreInsights/Source/IoStoreInsights/Private/ViewModels/IoStoreActivityTableTreeNode.h
@@ -0,0 +1,56 @@
+// Copyright Epic Games, Inc. All Rights Reserved.
+
+#pragma once
+
+#include "CoreTypes.h"
+#include "ViewModels/IoStoreActivityTable.h"
+#include "InsightsCore/Table/ViewModels/TableTreeNode.h"
+#include "InsightsCore/Table/ViewModels/Table.h"
+
+namespace UE::IoStoreInsights
+{
+	struct FIoStoreActivity;
+	class FIoStoreActivityNode;
+
+	// Type definition for shared pointers to instances of FIoStoreActivityNode
+	typedef TSharedPtr<class FIoStoreActivityNode> FActivityNodePtr;
+
+	// Type definition for shared references to instances of FIoStoreActivityNode
+	typedef TSharedRef<class FIoStoreActivityNode> FActivityNodeRef;
+
+	// Type definition for shared references to const instances of FIoStoreActivityNode
+	typedef TSharedRef<const class FIoStoreActivityNode> FActivityNodeRefConst;
+
+	// Type definition for weak references to instances of FIoStoreActivityNode
+	typedef TWeakPtr<class FIoStoreActivityNode> FActivityNodeWeak;
+
+	// Class used to store information about an IoStore activity (used in SActivityTreeView)
+	class FIoStoreActivityNode : public UE::Insights::FTableTreeNode
+	{
+		INSIGHTS_DECLARE_RTTI(FIoStoreActivityNode, UE::Insights::FTableTreeNode)
+
+	public:
+		explicit FIoStoreActivityNode(const FName InName, TWeakPtr<FIoStoreActivityTable> InParentTable, int32 InRowIndex)
+			: UE::Insights::FTableTreeNode(InName, InParentTable, InRowIndex)
+		{
+		}
+
+		explicit FIoStoreActivityNode(const FName InGroupName, TWeakPtr<FIoStoreActivityTable> InParentTable)
+			: UE::Insights::FTableTreeNode(InGroupName, InParentTable)
+		{
+		}
+
+		FIoStoreActivityTable& GetActivityTableChecked() const
+		{
+			const TSharedPtr<UE::Insights::FTable>& TablePin = GetParentTable().Pin();
+			check(TablePin.IsValid());
+			return *StaticCastSharedPtr<FIoStoreActivityTable>(TablePin);
+		}
+
+		bool IsValidActivity() const { return GetActivityTableChecked().IsValidRowIndex(GetRowIndex()); }
+		const FIoStoreActivity* GetActivity() const { return GetActivityTableChecked().GetActivity(GetRowIndex()); }
+		const FIoStoreActivity* GetActivityChecked() const { return GetActivityTableChecked().GetActivityChecked(GetRowIndex()); }
+	};
+
+
+} //UE::IoStoreInsights
diff --git a/Engine/Plugins/IoStoreInsights/Source/IoStoreInsights/Private/ViewModels/IoStoreInsightsTrack.cpp b/Engine/Plugins/IoStoreInsights/Source/IoStoreInsights/Private/ViewModels/IoStoreInsightsTrack.cpp
--- a/Engine/Plugins/IoStoreInsights/Source/IoStoreInsights/Private/ViewModels/IoStoreInsightsTrack.cpp
+++ b/Engine/Plugins/IoStoreInsights/Source/IoStoreInsights/Private/ViewModels/IoStoreInsightsTrack.cpp
@@ -80,10 +80,9 @@ namespace UE::IoStoreInsights
 				const double Duration = InEvent.EndTime - InEvent.StartTime;
 				InOutTooltip.AddNameValueTextLine(TEXT("Duration:"), Insights::FormatTimeAuto(Duration));
 
-				const TCHAR* PackageName = *IoStoreRequest.PackageName ? IoStoreRequest.PackageName : IoStoreActivityProvider->GetPackageName(IoStoreRequest.PackageId);
-				if (*PackageName)
+				if (*IoStoreRequest.PackageName)
 				{
-					InOutTooltip.AddNameValueTextLine(TEXT("Package:"), PackageName);
+					InOutTooltip.AddNameValueTextLine(TEXT("Package:"), IoStoreRequest.PackageName);
 				}
 				if (*IoStoreRequest.ExtraTag)
 				{
@@ -108,7 +107,7 @@ namespace UE::IoStoreInsights
 
 				// potentially add a callstack at the very bottom
 				static const bool bAlwaysShowCallstack = false; // @todo: put on the right-click menu
-				if (bAlwaysShowCallstack || (*PackageName == '\0' && *IoStoreRequest.ExtraTag == '\0')) // only show callstack if there's no package or custom tag
+				if (bAlwaysShowCallstack || (*IoStoreRequest.PackageName == '\0' && *IoStoreRequest.ExtraTag == '\0')) // only show callstack if there's no package or custom tag
 				{
 					// package id may be useful for debugging where the missing tag is
 					if (IoStoreRequest.PackageId != 0)
diff --git a/Engine/Plugins/IoStoreInsights/Source/IoStoreInsights/Private/Widgets/SIoStoreActivityTableTreeView.cpp b/Engine/Plugins/IoStoreInsights/Source/IoStoreInsights/Private/Widgets/SIoStoreActivityTableTreeView.cpp
new file mode 100644
--- /dev/null
+++ b/Engine/Plugins/IoStoreInsights/Source/IoStoreInsights/Private/Widgets/SIoStoreActivityTableTreeView.cpp
@@ -0,0 +1,101 @@
+// Copyright Epic Games, Inc. All Rights Reserved.
+
+#include "SIoStoreActivityTableTreeView.h"
+#include "Model/IoStoreInsightsProvider.h"
+#include "ViewModels/IoStoreActivityTable.h"
+#include "ViewModels/IoStoreActivityTableTreeNode.h"
+
+
+namespace UE::IoStoreInsights
+{
+	void SActivityTableTreeView::Construct(const FArguments& InArgs, TSharedPtr<FIoStoreActivityTable> InTablePtr)
+	{
+		ConstructWidget(InTablePtr);
+		CreateGroupings();
+		CreateSortings();
+	}
+
+
+
+	void SActivityTableTreeView::Tick(const FGeometry& AllottedGeometry, const double InCurrentTime, const float InDeltaTime)
+	{
+		STableTreeView::Tick(AllottedGeometry, InCurrentTime, InDeltaTime);
+
+		if (!bIsUpdateRunning)
+		{
+			RebuildTree(false);
+		}
+	}
+
+
+
+	void SActivityTableTreeView::RebuildTree(bool bResync)
+	{
+		if (StartTime > EndTime)
+		{
+			bRangeDirty = false;
+			return;
+		}
+
+		if (bResync || bRangeDirty)
+		{
+			bRangeDirty = false;
+
+			TSharedPtr<FIoStoreActivityTable> ActivityTable = GetActivityTable();
+			TArray<const FIoStoreActivity*>& Activities = ActivityTable->GetActivities();
+			Activities.Empty();
+			TableRowNodes.Empty();
+
+			if (StartTime < EndTime && AnalysisSession != nullptr)
+			{
+				const FIoStoreInsightsProvider* Provider = AnalysisSession->ReadProvider<FIoStoreInsightsProvider>(IIoStoreInsightsProvider::ProviderName);
+				if (Provider)
+				{
+					Provider->EnumerateIoStoreRequests([this, &ActivityTable, &Activities](const FIoStoreRequest& IoStoreRequest, const IIoStoreInsightsProvider::Timeline& Timeline)
+					{
+						Timeline.EnumerateEvents(StartTime, EndTime, [this, &ActivityTable, &Activities](double EventStartTime, double EventEndTime, uint32 EventDepth, const FIoStoreActivity* IoStoreActivity)
+						{
+							if (IoStoreActivity->ActivityType == EIoStoreActivityType::Request_Read && IoStoreActivity->EndTime > 0 && !IoStoreActivity->Failed)
+							{
+								static const FName BaseNodeName(TEXT("row"));
+
+								Activities.Emplace(IoStoreActivity);
+								uint32 Index = Activities.Num() - 1;
+								FName NodeName(BaseNodeName, static_cast<int32>(IoStoreActivity->IoStoreRequest->IoStoreRequestIndex+1));
+								FActivityNodePtr NodePtr = MakeShared<FIoStoreActivityNode>(NodeName, ActivityTable, Index);
+								TableRowNodes.Add(NodePtr);
+							}
+
+							return TraceServices::EEventEnumerate::Continue;
+						});
+
+						return true;
+					});
+				}
+			}
+
+			UpdateTree();
+			TreeView->RebuildList();
+		}
+	}
+
+
+
+	void SActivityTableTreeView::SetRange(double InStartTime, double InEndTime)
+	{
+		bRangeDirty = (StartTime != InStartTime) || (EndTime != InEndTime);
+
+		StartTime = InStartTime;
+		EndTime = InEndTime;
+	}
+
+
+
+	void SActivityTableTreeView::SetAnalysisSession(const TraceServices::IAnalysisSession* InAnalysisSession)
+	{
+		AnalysisSession = InAnalysisSession;
+	}
+
+
+
+} // UE::IoStoreInsights
diff --git a/Engine/Plugins/IoStoreInsights/Source/IoStoreInsights/Private/Widgets/SIoStoreActivityTableTreeView.h b/Engine/Plugins/IoStoreInsights/Source/IoStoreInsights/Private/Widgets/SIoStoreActivityTableTreeView.h
new file mode 100644
--- /dev/null
+++ b/Engine/Plugins/IoStoreInsights/Source/IoStoreInsights/Private/Widgets/SIoStoreActivityTableTreeView.h
@@ -0,0 +1,39 @@
+// Copyright Epic Games, Inc. All Rights Reserved.
+
+#pragma once
+
+#include "CoreMinimal.h"
+#include "InsightsCore/Table/Widgets/STableTreeView.h"
+#include "ViewModels/IoStoreActivityTable.h"
+
+namespace TraceServices { class IAnalysisSession; }
+
+namespace UE::IoStoreInsights
+{
+	class SActivityTableTreeView : public UE::Insights::STableTreeView
+	{
+	public:
+		virtual ~SActivityTableTreeView() = default;
+
+		SLATE_BEGIN_ARGS(SActivityTableTreeView) {}
+		SLATE_END_ARGS()
+
+		void Construct(const FArguments& InArgs, TSharedPtr<FIoStoreActivityTable> InTablePtr);
+
+		TSharedPtr<FIoStoreActivityTable> GetActivityTable() { return StaticCastSharedPtr<FIoStoreActivityTable>(GetTable()); }
+		const TSharedPtr<FIoStoreActivityTable> GetActivityTable() const { return StaticCastSharedPtr<FIoStoreActivityTable>(GetTable()); }
+
+		void SetRange(double InStartTime, double InEndTime);
+		void SetAnalysisSession(const TraceServices::IAnalysisSession* InAnalysisSession);
+
+		virtual void Tick(const FGeometry& AllottedGeometry, const double InCurrentTime, const float InDeltaTime) override;
+		virtual void RebuildTree(bool bResync);
+
+	private:
+		bool bRangeDirty = false;
+		double StartTime = 0.0;
+		double EndTime = -1.0;
+		const TraceServices::IAnalysisSession* AnalysisSession = nullptr;
+	};
+
+} //UE::IoStoreInsights
diff --git a/Engine/Plugins/IoStoreInsights/Source/IoStoreInsights/Private/Widgets/SIoStoreAnalysisActivityView.cpp b/Engine/Plugins/IoStoreInsights/Source/IoStoreInsights/Private/Widgets/SIoStoreAnalysisActivityView.cpp
deleted file mode 100644
diff --git a/Engine/Plugins/IoStoreInsights/Source/IoStoreInsights/Private/Widgets/SIoStoreAnalysisActivityView.h b/Engine/Plugins/IoStoreInsights/Source/IoStoreInsights/Private/Widgets/SIoStoreAnalysisActivityView.h
deleted file mode 100644
diff --git a/Engine/Plugins/IoStoreInsights/Source/IoStoreInsights/Private/Widgets/SIoStoreAnalysisTab.cpp b/Engine/Plugins/IoStoreInsights/Source/IoStoreInsights/Private/Widgets/SIoStoreAnalysisTab.cpp
--- a/Engine/Plugins/IoStoreInsights/Source/IoStoreInsights/Private/Widgets/SIoStoreAnalysisTab.cpp
+++ b/Engine/Plugins/IoStoreInsights/Source/IoStoreInsights/Private/Widgets/SIoStoreAnalysisTab.cpp
@@ -2,7 +2,7 @@
 
 #include "SIoStoreAnalysisTab.h"
 #include "SIoStoreAnalysisReadSizeHistogramView.h"
-#include "SIoStoreAnalysisActivityView.h"
+#include "SIoStoreActivityTableTreeView.h"
 #include "Model/IoStoreInsightsProvider.h"
 #include "ViewModels/IoStoreInsightsTrack.h"
 #include "Widgets/Text/STextBlock.h"
@@ -68,16 +68,17 @@ namespace UE::IoStoreInsights
 		};
 
 		// add the view modes
-		AddViewMode(EViewMode::ReadSizes, LOCTEXT("ViewMode_ReadSizes", "Read Size Histogram"))
+		TSharedRef<FIoStoreActivityTable> ActivityTable = MakeShared<FIoStoreActivityTable>();
+		ActivityTable->Reset();
+		AddViewMode(EViewMode::ReadActivity, LOCTEXT("ViewMode_ReadActivity", "Read Activity"))
 		[
-			SAssignNew(ReadSizeHistogramView, SIoStoreAnalysisReadSizeHistogramView)
-				.ListItemsSource(&ReadSizeHistogramItems)
+			SAssignNew(ActivityTableTreeView, SActivityTableTreeView, ActivityTable )
 		];
 
-		AddViewMode(EViewMode::ReadActivity, LOCTEXT("ViewMode_ReadActivity", "Read Activity"))
+		AddViewMode(EViewMode::ReadSizes, LOCTEXT("ViewMode_ReadSizes", "Read Size Histogram"))
 		[
-			SAssignNew(ActivityInfoTreeView, SIoStoreAnalysisActivityView)
-				.TreeItemsSource(&ActivityInfoArray)
+			SAssignNew(ReadSizeHistogramView, SIoStoreAnalysisReadSizeHistogramView)
+				.ListItemsSource(&ReadSizeHistogramItems)
 		];
 
 		// update the data
@@ -115,7 +116,7 @@ namespace UE::IoStoreInsights
 		AnalysisSession = InAnalysisSession;
 		SharedStatePtr = InSharedStatePtr;
 		ReadSizeHistogramView->SetAnalysisSession(AnalysisSession);
-		ActivityInfoTreeView->SetAnalysisSession(AnalysisSession);
+		ActivityTableTreeView->SetAnalysisSession(AnalysisSession);
 
 		if (InTimingViewSession)
 		{
@@ -223,11 +224,8 @@ namespace UE::IoStoreInsights
 			}
 		}
 
-
 		ReadSizeHistogramView->RebuildList();
-
-		ActivityInfoTreeView->ApplySort();
-		ActivityInfoTreeView->RebuildList();
+		ActivityTableTreeView->SetRange(StartTime, EndTime);
 	}
 
 
@@ -243,13 +241,9 @@ namespace UE::IoStoreInsights
 		TMap<uint64,TSharedPtr<Private::FReadSizeHistogramItem>> QuantizedSizeToHisogramItemMap;
 		ReadSizeHistogramItems.Reset();
 
-		TMap<uint32,TArray<const FIoStoreActivity*>> DataSourceToActivityMap;
-		ActivityInfoArray.Reset();
-
-
-		Provider->EnumerateIoStoreRequests([this, &DataSourceToActivityMap, &TotalReads, &MinQuantizedReadSize, &MaxQuantizedReadSize, &QuantizedSizeToHisogramItemMap](const FIoStoreRequest& IoStoreRequest, const IIoStoreInsightsProvider::Timeline& Timeline)
+		Provider->EnumerateIoStoreRequests([this, &TotalReads, &MinQuantizedReadSize, &MaxQuantizedReadSize, &QuantizedSizeToHisogramItemMap](const FIoStoreRequest& IoStoreRequest, const IIoStoreInsightsProvider::Timeline& Timeline)
 		{
-			Timeline.EnumerateEvents(StartTime, EndTime, [this, &DataSourceToActivityMap, &TotalReads, &MinQuantizedReadSize, &MaxQuantizedReadSize, &QuantizedSizeToHisogramItemMap](double EventStartTime, double EventEndTime, uint32 EventDepth, const FIoStoreActivity* IoStoreActivity)
+			Timeline.EnumerateEvents(StartTime, EndTime, [this, &TotalReads, &MinQuantizedReadSize, &MaxQuantizedReadSize, &QuantizedSizeToHisogramItemMap](double EventStartTime, double EventEndTime, uint32 EventDepth, const FIoStoreActivity* IoStoreActivity)
 			{
 				if (IoStoreActivity->ActivityType == EIoStoreActivityType::Request_Read && IoStoreActivity->EndTime > 0 && !IoStoreActivity->Failed)
 				{
@@ -279,22 +273,6 @@ namespace UE::IoStoreInsights
 					HistogramItem->Count++;
 					HistogramItem->MinDuration = FMath::Min(Duration, HistogramItem->MinDuration);
 					HistogramItem->MaxDuration = FMath::Max(Duration, HistogramItem->MaxDuration);
-
-					// collect the activities associated with the package or other data source
-					uint32 SourceHash;
-					if (*IoStoreActivity->IoStoreRequest->PackageName)
-					{
-						SourceHash = PointerHash(IoStoreActivity->IoStoreRequest->PackageName);
-					}
-					else if (*IoStoreActivity->IoStoreRequest->ExtraTag)
-					{
-						SourceHash = PointerHash(IoStoreActivity->IoStoreRequest->ExtraTag);
-					}
-					else
-					{
-						SourceHash = GetTypeHash(IoStoreActivity->IoStoreRequest->PackageId);
-					}
-					DataSourceToActivityMap.FindOrAdd(SourceHash).Add(IoStoreActivity);
 				}
 
 				return TraceServices::EEventEnumerate::Continue;
@@ -333,65 +311,6 @@ namespace UE::IoStoreInsights
 		{
 			HistogramItem->CountAsPctNormalized = HistogramItem->CountAsPct * (1.0f / MaxPct);
 		}
-
-		// build the tree
-		for (const auto& Itr : DataSourceToActivityMap)
-		{
-			const TArray<const FIoStoreActivity*>& Activities = Itr.Value;
-			if (Activities.Num() == 0)
-			{
-				continue;
-			}
-			const FIoStoreActivity* FirstActivity = Activities[0];
-
-			// prepare parent information, based on the first activity
-			FStringBuilderBase PackageNameAndTag;
-			if (*FirstActivity->IoStoreRequest->PackageName)
-			{
-				PackageNameAndTag.Append(FirstActivity->IoStoreRequest->PackageName);
-			}
-			else if (IoStoreActivityProvider != nullptr)
-			{
-				PackageNameAndTag.Append(IoStoreActivityProvider->GetPackageName(FirstActivity->IoStoreRequest->PackageId));
-			}
-			else
-			{
-				PackageNameAndTag.Appendf(TEXT("%x"), FirstActivity->IoStoreRequest->PackageId);
-			}
-			if (*FirstActivity->IoStoreRequest->ExtraTag)
-			{
-				PackageNameAndTag.Append(TEXT(" Tag: "));
-				PackageNameAndTag.Append(FirstActivity->IoStoreRequest->ExtraTag);
-			}
-			if (Itr.Value.Num() > 1)
-			{
-				PackageNameAndTag.Appendf( TEXT(" (%d)"), Itr.Value.Num());
-			}
-
-			TSharedPtr<Private::FActivityInfoParent> Parent = MakeShared<Private::FActivityInfoParent>();
-			Parent->DisplayName = PackageNameAndTag.ToString();
-			Parent->IoStoreActivity = FirstActivity;
-			Parent->bIsParent = true;
-			ActivityInfoArray.Add(Parent);
-
-			// process all activities
-			for (const FIoStoreActivity* Activity : Activities)
-			{
-				double Duration = (Activity->EndTime > 0) ? (Activity->EndTime - Activity->StartTime) : +std::numeric_limits<double>::infinity();
-				Parent->Duration += Duration;
-				Parent->TotalSize += Activity->ActualSize;
-
-				// only add child items if there's more than one - we'll display the information in the parent if there's only one
-				if (Itr.Value.Num() > 1)
-				{
-					TSharedPtr<Private::FActivityInfoChild> Child = MakeShared<Private::FActivityInfoChild>();
-					Child->IoStoreActivity = Activity;
-					Child->Duration = Duration;
-					Child->bIsParent = false;
-					Parent->Children.Add(Child);
-				}
-			}
-		}
 	}
 
 } //namespace UE::IoStoreInsights
diff --git a/Engine/Plugins/IoStoreInsights/Source/IoStoreInsights/Private/Widgets/SIoStoreAnalysisTab.h b/Engine/Plugins/IoStoreInsights/Source/IoStoreInsights/Private/Widgets/SIoStoreAnalysisTab.h
--- a/Engine/Plugins/IoStoreInsights/Source/IoStoreInsights/Private/Widgets/SIoStoreAnalysisTab.h
+++ b/Engine/Plugins/IoStoreInsights/Source/IoStoreInsights/Private/Widgets/SIoStoreAnalysisTab.h
@@ -20,12 +20,11 @@ namespace TraceServices
 
 namespace UE::IoStoreInsights
 {
-	class SIoStoreAnalysisActivityView;
 	class SIoStoreAnalysisReadSizeHistogramView;
+	class SActivityTableTreeView;
 	class FIoStoreInsightsProvider;
 	class FIoStoreInsightsViewSharedState;
 	namespace Private { struct FReadSizeHistogramItem; }
-	namespace Private { struct FActivityInfo; }
 
 	class SIoStoreAnalysisTab : public SCompoundWidget
 	{
@@ -42,8 +41,8 @@ namespace UE::IoStoreInsights
 	private:
 		enum class EViewMode : uint8
 		{
-			ReadSizes,
 			ReadActivity,
+			ReadSizes,
 		};
 
 		void HandleTimeMarkerChanged(UE::Insights::Timing::ETimeChangedFlags InFlags, double InTimeMarker);
@@ -64,13 +63,12 @@ namespace UE::IoStoreInsights
 
 		double StartTime = -1;
 		double EndTime = -1;
-		EViewMode ViewMode = EViewMode::ReadSizes;
+		EViewMode ViewMode = EViewMode::ReadActivity;
 
 		TArray<TSharedPtr<Private::FReadSizeHistogramItem>> ReadSizeHistogramItems;
 		TSharedPtr<SIoStoreAnalysisReadSizeHistogramView> ReadSizeHistogramView;
 
-		TArray<TSharedPtr<Private::FActivityInfo>> ActivityInfoArray;
-		TSharedPtr<SIoStoreAnalysisActivityView> ActivityInfoTreeView;
+		TSharedPtr<SActivityTableTreeView> ActivityTableTreeView;
 	};
 
 } //namespace UE::IoStoreInsights
