Add Backend Name to the IoStore insights table

 - also remove backend and package name lookup functions from the IoStore insights provider as they're no longer necessary

#rnx

diff --git a/Engine/Plugins/IoStoreInsights/Source/IoStoreInsights/Private/Model/IoStoreInsightsProvider.cpp b/Engine/Plugins/IoStoreInsights/Source/IoStoreInsights/Private/Model/IoStoreInsightsProvider.cpp
--- a/Engine/Plugins/IoStoreInsights/Source/IoStoreInsights/Private/Model/IoStoreInsightsProvider.cpp
+++ b/Engine/Plugins/IoStoreInsights/Source/IoStoreInsights/Private/Model/IoStoreInsightsProvider.cpp
@@ -70,13 +70,25 @@ namespace UE::IoStoreInsights
 		FIoStoreRequestInfoInternal& IoStoreRequestInfo = IoStoreRequests[IoStoreRequestIndex];
 		FIoStoreActivity& IoStoreActivity = IoStoreRequestStates.PushBack();
 		IoStoreActivity.IoStoreRequest = &IoStoreRequestInfo.IoStoreRequestInfo;
-		IoStoreActivity.BackendHandle = BackendHandle;
 		IoStoreActivity.ActualSize = 0;
 		IoStoreActivity.StartTime = Time;
 		IoStoreActivity.EndTime = std::numeric_limits<double>::infinity();
 		IoStoreActivity.ThreadId = ThreadId;
 		IoStoreActivity.ActivityType = Type;
 		IoStoreActivity.Failed = false;
+
+		const TCHAR** BackendNamePtr = BackendNameMap.Find(BackendHandle);
+		if (BackendNamePtr)
+		{
+			IoStoreActivity.BackendName = (*BackendNamePtr);
+		}
+		else
+		{
+			IoStoreActivity.BackendName = TEXT("(Unknown)");
+			TArray<FIoStoreActivity*> PendingBackendNameItems = PendingBackendNameMap.FindOrAdd(BackendHandle);
+			PendingBackendNameItems.Add(&IoStoreActivity);
+		}
+		
 	
 		return IoStoreRequestInfo.ActivityTimeline->AppendBeginEvent(Time, &IoStoreActivity);
 	}
@@ -116,30 +128,20 @@ namespace UE::IoStoreInsights
 
 
 
-	const TCHAR* FIoStoreInsightsProvider::GetPackageName(uint64 PackageId) const
-	{
-		const TCHAR* const* PackageNamePtr = PackageMap.Find(PackageId);
-		if (PackageNamePtr)
-		{
-			return *PackageNamePtr;
-		}
-
-		return TEXT("");
-	}
-
-
-
 	void FIoStoreInsightsProvider::AddBackendName(uint64 BackendHandle, const TCHAR* BackendName)
 	{
 		BackendNameMap.Add(BackendHandle, BackendName);
-	}
-
 
+		const TArray<FIoStoreActivity*>* PendingBackendNameItemsPtr = PendingBackendNameMap.Find(BackendHandle);
+		if (PendingBackendNameItemsPtr != nullptr && *BackendName != '\0')
+		{
+			for ( FIoStoreActivity* IoStoreActivity : (*PendingBackendNameItemsPtr) )
+			{
+				IoStoreActivity->BackendName = BackendName;
+			}
 
-	const TCHAR* FIoStoreInsightsProvider::GetBackendName(uint64 BackendHandle) const
-	{
-		const TCHAR* const* BackendNamePtr = BackendNameMap.Find(BackendHandle);
-		return BackendNamePtr ? *BackendNamePtr : TEXT("unknown backend");
+			PendingBackendNameMap.Remove(BackendHandle);
+		}
 	}
 
 
diff --git a/Engine/Plugins/IoStoreInsights/Source/IoStoreInsights/Private/Model/IoStoreInsightsProvider.h b/Engine/Plugins/IoStoreInsights/Source/IoStoreInsights/Private/Model/IoStoreInsightsProvider.h
--- a/Engine/Plugins/IoStoreInsights/Source/IoStoreInsights/Private/Model/IoStoreInsightsProvider.h
+++ b/Engine/Plugins/IoStoreInsights/Source/IoStoreInsights/Private/Model/IoStoreInsightsProvider.h
@@ -36,10 +36,8 @@ namespace UE::IoStoreInsights
 		void EndIoStoreActivity(uint32 IoStoreRequestIndex, uint64 ActivityIndex, uint64 ActualSize, bool Failed, double Time);
 
 		void AddPackageMapping(uint64 PackageId, const TCHAR* PackageName);
-		virtual const TCHAR* GetPackageName(uint64 PackageId) const override;
 
 		void AddBackendName(uint64 BackendHandle, const TCHAR* BackendName);
-		virtual const TCHAR* GetBackendName(uint64 BackendHandle) const override;
 
 		virtual const FIoStoreRequest& GetIoStoreRequest(uint32 IoStoreRequestIndex) const override;
 
@@ -57,6 +55,7 @@ namespace UE::IoStoreInsights
 		TMap<uint64, const TCHAR*> BackendNameMap;
 
 		TMap<uint64,TArray<uint32>> PendingPackageNameMap;
+		TMap<uint64,TArray<FIoStoreActivity*>> PendingBackendNameMap;
 	};
 
 } // namespace UE::IoStoreInsights
diff --git a/Engine/Plugins/IoStoreInsights/Source/IoStoreInsights/Private/ViewModels/IoStoreActivityTable.cpp b/Engine/Plugins/IoStoreInsights/Source/IoStoreInsights/Private/ViewModels/IoStoreActivityTable.cpp
--- a/Engine/Plugins/IoStoreInsights/Source/IoStoreInsights/Private/ViewModels/IoStoreActivityTable.cpp
+++ b/Engine/Plugins/IoStoreInsights/Source/IoStoreInsights/Private/ViewModels/IoStoreActivityTable.cpp
@@ -23,6 +23,7 @@ namespace UE::IoStoreInsights
 	const FName FActivityTableColumns::ColumnRequestChunkId("ChunkId");
 	const FName FActivityTableColumns::ColumnRequestChunkType("ChunkType");
 	const FName FActivityTableColumns::ColumnRequestStartTime("StartTime");
+	const FName FActivityTableColumns::ColumnRequestBackend("Backend");
 
 	typedef FTableCellValue (*FActivityFieldGetterFn) (const FTableColumn&, const FIoStoreActivity*);
 
@@ -80,6 +81,7 @@ namespace UE::IoStoreInsights
 		static FTableCellValue GetChunkId  (const FTableColumn& Column, const FIoStoreActivity* Activity) { return FTableCellValue((int64)Activity->IoStoreRequest->ChunkIdHash); }
 		static FTableCellValue GetChunkType(const FTableColumn& Column, const FIoStoreActivity* Activity) { return FTableCellValue(FText::FromString(LexToString((EIoChunkType)Activity->IoStoreRequest->ChunkType))); }
 		static FTableCellValue GetStartTime(const FTableColumn& Column, const FIoStoreActivity* Activity) { return FTableCellValue(Activity->StartTime); }
+		static FTableCellValue GetBackend  (const FTableColumn& Column, const FIoStoreActivity* Activity) { return FTableCellValue(Activity->BackendName); }
 	};
 
 
@@ -232,6 +234,22 @@ namespace UE::IoStoreInsights
 			Column.SetValueSorter(MakeShared<FSorterByDoubleValue>(ColumnRef));
 			Column.SetAggregation(ETableColumnAggregation::Min);
 		}
+
+		// Backend Column
+		{
+			TSharedRef<FTableColumn> ColumnRef = MakeColumn(FActivityTableColumns::ColumnRequestBackend,
+															LOCTEXT("BackendName",  "Backend"),
+															LOCTEXT("BackendTitle", "Backend"),
+															LOCTEXT("BackendDesc",  "IoDispatcher Backend that handled the request"),
+															75.0f);
+			FTableColumn& Column = *ColumnRef;
+			Column.SetFlags(ETableColumnFlags::ShouldBeVisible | ETableColumnFlags::CanBeHidden | ETableColumnFlags::CanBeFiltered);
+			Column.SetDataType(ETableCellDataType::CString);
+			Column.SetValueGetter(MakeShared<FActivityColumnValueGetter<DefaultActivityFieldGetterFuncs::GetBackend>>());
+			Column.SetValueFormatter(MakeShared<FCStringValueFormatterAsText>());
+			Column.SetValueSorter(MakeShared<FSorterByCStringValue>(ColumnRef));
+			Column.SetAggregation(ETableColumnAggregation::SameValue);
+		}
 	}
 
 } //namespace UE::IoStoreInsights
diff --git a/Engine/Plugins/IoStoreInsights/Source/IoStoreInsights/Private/ViewModels/IoStoreActivityTable.h b/Engine/Plugins/IoStoreInsights/Source/IoStoreInsights/Private/ViewModels/IoStoreActivityTable.h
--- a/Engine/Plugins/IoStoreInsights/Source/IoStoreInsights/Private/ViewModels/IoStoreActivityTable.h
+++ b/Engine/Plugins/IoStoreInsights/Source/IoStoreInsights/Private/ViewModels/IoStoreActivityTable.h
@@ -22,6 +22,7 @@ namespace UE::IoStoreInsights
 		static const FName ColumnRequestChunkId;
 		static const FName ColumnRequestChunkType;
 		static const FName ColumnRequestStartTime;
+		static const FName ColumnRequestBackend;
 	};
 
 	class FIoStoreActivityTable : public UE::Insights::FTable
diff --git a/Engine/Plugins/IoStoreInsights/Source/IoStoreInsights/Private/ViewModels/IoStoreInsightsTimingViewExtender.cpp b/Engine/Plugins/IoStoreInsights/Source/IoStoreInsights/Private/ViewModels/IoStoreInsightsTimingViewExtender.cpp
--- a/Engine/Plugins/IoStoreInsights/Source/IoStoreInsights/Private/ViewModels/IoStoreInsightsTimingViewExtender.cpp
+++ b/Engine/Plugins/IoStoreInsights/Source/IoStoreInsights/Private/ViewModels/IoStoreInsightsTimingViewExtender.cpp
@@ -96,7 +96,7 @@ namespace UE::IoStoreInsights
 							TimingEventType = HashCombine(TimingEventType, IoStoreRequest.PackageId);
 
 							uint32 Type = ((uint32)IoStoreActivity->ActivityType & 0x0F) | (IoStoreActivity->Failed ? 0x80 : 0);
-							AllIoEvents.Add(FIoStoreEventState{ RequestState, EventStartTime, EventEndTime, IoStoreActivity->ActualSize, IoStoreActivity->BackendHandle, TimingEventType, LocalDepth, Type });
+							AllIoEvents.Add(FIoStoreEventState{ RequestState, EventStartTime, EventEndTime, IoStoreActivity->ActualSize, IoStoreActivity->BackendName, TimingEventType, LocalDepth, Type });
 							return TraceServices::EEventEnumerate::Continue;
 						});
 
diff --git a/Engine/Plugins/IoStoreInsights/Source/IoStoreInsights/Private/ViewModels/IoStoreInsightsTimingViewExtender.h b/Engine/Plugins/IoStoreInsights/Source/IoStoreInsights/Private/ViewModels/IoStoreInsightsTimingViewExtender.h
--- a/Engine/Plugins/IoStoreInsights/Source/IoStoreInsights/Private/ViewModels/IoStoreInsightsTimingViewExtender.h
+++ b/Engine/Plugins/IoStoreInsights/Source/IoStoreInsights/Private/ViewModels/IoStoreInsightsTimingViewExtender.h
@@ -35,7 +35,7 @@ namespace UE::IoStoreInsights
 		double StartTime;
 		double EndTime;
 		uint64 ActualSize;
-		uint64 BackendHandle;
+		const TCHAR* BackendName;
 		uint64 TimingEventType;
 		uint32 Depth; // During update, this is local within a track - then it's set to a global depth
 		uint32 Type; // EIoStoreActivityType + "Failed" flag
diff --git a/Engine/Plugins/IoStoreInsights/Source/IoStoreInsights/Private/ViewModels/IoStoreInsightsTrack.cpp b/Engine/Plugins/IoStoreInsights/Source/IoStoreInsights/Private/ViewModels/IoStoreInsightsTrack.cpp
--- a/Engine/Plugins/IoStoreInsights/Source/IoStoreInsights/Private/ViewModels/IoStoreInsightsTrack.cpp
+++ b/Engine/Plugins/IoStoreInsights/Source/IoStoreInsights/Private/ViewModels/IoStoreInsightsTrack.cpp
@@ -98,7 +98,7 @@ namespace UE::IoStoreInsights
 						InOutTooltip.AddNameValueTextLine(TEXT("Size:"), FText::AsNumber(IoStoreRequest.Size).ToString() + TEXT(" bytes"));
 					}
 					InOutTooltip.AddNameValueTextLine(TEXT("Result Size:"), FText::AsNumber(InEvent.ActualSize).ToString() + TEXT(" bytes"));
-					InOutTooltip.AddNameValueTextLine(TEXT("Backend:"), IoStoreActivityProvider->GetBackendName(InEvent.BackendHandle));
+					InOutTooltip.AddNameValueTextLine(TEXT("Backend:"), InEvent.BackendName);
 				}
 				else
 				{
@@ -193,8 +193,7 @@ namespace UE::IoStoreInsights
 		}
 		else if (ActivityType == EIoStoreActivityType::Request_Read)
 		{
-			const TCHAR* BackendName = IoStoreActivityProvider->GetBackendName(Event.BackendHandle);
-			Color =  FTimingEvent::ComputeEventColor(BackendName);
+			Color =  FTimingEvent::ComputeEventColor(Event.BackendName);
 		}
 		else
 		{
diff --git a/Engine/Plugins/IoStoreInsights/Source/IoStoreInsights/Public/IIoStoreInsightsProvider.h b/Engine/Plugins/IoStoreInsights/Source/IoStoreInsights/Public/IIoStoreInsightsProvider.h
--- a/Engine/Plugins/IoStoreInsights/Source/IoStoreInsights/Public/IIoStoreInsightsProvider.h
+++ b/Engine/Plugins/IoStoreInsights/Source/IoStoreInsights/Public/IIoStoreInsightsProvider.h
@@ -36,7 +36,7 @@ namespace UE::IoStoreInsights
 		double StartTime = 0.0;
 		double EndTime = 0.0;
 		uint64 ActualSize = 0;
-		uint64 BackendHandle = 0;
+		const TCHAR* BackendName = nullptr;
 		uint32 ThreadId = 0;
 		EIoStoreActivityType ActivityType = EIoStoreActivityType::Invalid;
 		bool Failed = false;
@@ -52,8 +52,6 @@ namespace UE::IoStoreInsights
 		virtual ~IIoStoreInsightsProvider() = default;
 		virtual void EnumerateIoStoreRequests(TFunctionRef<bool(const FIoStoreRequest&, const Timeline&)> Callback) const = 0;
 
-		virtual const TCHAR* GetPackageName(uint64 PackageId) const = 0;
-		virtual const TCHAR* GetBackendName(uint64 BackendHandle) const = 0;
 		virtual const FIoStoreRequest& GetIoStoreRequest(uint32 IoStoreRequestIndex) const = 0;
 	};
 
