Remove assert for duplicate IoStore activities

#rnx

diff --git a/Engine/Plugins/IoStoreInsights/Source/IoStoreInsights/Private/Analyzers/IoStoreInsightsAnalyzer.cpp b/Engine/Plugins/IoStoreInsights/Source/IoStoreInsights/Private/Analyzers/IoStoreInsightsAnalyzer.cpp
--- a/Engine/Plugins/IoStoreInsights/Source/IoStoreInsights/Private/Analyzers/IoStoreInsightsAnalyzer.cpp
+++ b/Engine/Plugins/IoStoreInsights/Source/IoStoreInsights/Private/Analyzers/IoStoreInsightsAnalyzer.cpp
@@ -66,15 +66,20 @@ namespace UE::IoStoreInsights
 				const uint32 ChunkIdHash   = EventData.GetValue<uint32>("ChunkIdHash");
 				const uint8 ChunkType      = EventData.GetValue<uint8>("ChunkType");
 				const uint32 ThreadId      = Context.ThreadInfo.GetId();
-				check( !ActiveRequestsMap.Contains(RequestHandle) );
-
-				FString PackageName, ExtraTag;
-				uint64 PackageId;
-				GetPackageDetailFromMetadata(Context, ThreadId, PackageName, ExtraTag, PackageId);
-				uint32 IoStoreRequestIndex = Provider.GetIoStoreRequestIndex(ChunkIdHash, ChunkType, Offset, Size, CallstackId, PackageId, *PackageName, *ExtraTag);
+				if (!ActiveRequestsMap.Contains(RequestHandle))
+				{
+					FString PackageName, ExtraTag;
+					uint64 PackageId;
+					GetPackageDetailFromMetadata(Context, ThreadId, PackageName, ExtraTag, PackageId);
+					uint32 IoStoreRequestIndex = Provider.GetIoStoreRequestIndex(ChunkIdHash, ChunkType, Offset, Size, CallstackId, PackageId, *PackageName, *ExtraTag);
 
-				uint64 ActivityIndex = Provider.BeginIoStoreActivity(IoStoreRequestIndex, EIoStoreActivityType::Request_Pending, ThreadId, 0, Time);
-				ActiveRequestsMap.Emplace( RequestHandle, {IoStoreRequestIndex, ActivityIndex} );
+					uint64 ActivityIndex = Provider.BeginIoStoreActivity(IoStoreRequestIndex, EIoStoreActivityType::Request_Pending, ThreadId, 0, Time);
+					ActiveRequestsMap.Emplace( RequestHandle, {IoStoreRequestIndex, ActivityIndex} );
+				}
+				else
+				{
+					//IOSTORETRACE_WARNING("RequestCreate: Request already started!");
+				}
 			}
 			break;
 
