[IoStoreInsights] Fixed deprecated FName for FTableTreeNode constructor.

#rnx

#Synced-CL 51291823

diff --git a/Engine/Plugins/IoStoreInsights/Source/IoStoreInsights/Private/ViewModels/IoStoreActivityTableTreeNode.cpp b/Engine/Plugins/IoStoreInsights/Source/IoStoreInsights/Private/ViewModels/IoStoreActivityTableTreeNode.cpp
--- a/Engine/Plugins/IoStoreInsights/Source/IoStoreInsights/Private/ViewModels/IoStoreActivityTableTreeNode.cpp
+++ b/Engine/Plugins/IoStoreInsights/Source/IoStoreInsights/Private/ViewModels/IoStoreActivityTableTreeNode.cpp
@@ -4,6 +4,12 @@
 
 namespace UE::IoStoreInsights
 {
-	INSIGHTS_IMPLEMENT_RTTI(FIoStoreActivityNode)
+
+INSIGHTS_IMPLEMENT_RTTI(FIoStoreActivityNode)
+
+const FText FIoStoreActivityNode::GetDisplayName() const
+{
+	return FText::FromString(FString::Printf(TEXT("request_%u"), RequestIndex));
+}
 
 } // UE::IoStoreInsights
diff --git a/Engine/Plugins/IoStoreInsights/Source/IoStoreInsights/Private/ViewModels/IoStoreActivityTableTreeNode.h b/Engine/Plugins/IoStoreInsights/Source/IoStoreInsights/Private/ViewModels/IoStoreActivityTableTreeNode.h
--- a/Engine/Plugins/IoStoreInsights/Source/IoStoreInsights/Private/ViewModels/IoStoreActivityTableTreeNode.h
+++ b/Engine/Plugins/IoStoreInsights/Source/IoStoreInsights/Private/ViewModels/IoStoreActivityTableTreeNode.h
@@ -30,15 +30,13 @@ namespace UE::IoStoreInsights
 		INSIGHTS_DECLARE_RTTI(FIoStoreActivityNode, UE::Insights::FTableTreeNode)
 
 	public:
-		explicit FIoStoreActivityNode(const FName InName, TWeakPtr<FIoStoreActivityTable> InParentTable, int32 InRowIndex)
-			: UE::Insights::FTableTreeNode(InName, InParentTable, InRowIndex)
+		explicit FIoStoreActivityNode(TWeakPtr<FIoStoreActivityTable> InParentTable, int32 InRowIndex, uint32 InRequestIndex)
+			: UE::Insights::FTableTreeNode(InParentTable, InRowIndex)
+			, RequestIndex(InRequestIndex)
 		{
 		}
 
-		explicit FIoStoreActivityNode(const FName InGroupName, TWeakPtr<FIoStoreActivityTable> InParentTable)
-			: UE::Insights::FTableTreeNode(InGroupName, InParentTable)
-		{
-		}
+		virtual const FText GetDisplayName() const override;
 
 		FIoStoreActivityTable& GetActivityTableChecked() const
 		{
@@ -50,7 +48,9 @@ namespace UE::IoStoreInsights
 		bool IsValidActivity() const { return GetActivityTableChecked().IsValidRowIndex(GetRowIndex()); }
 		const FIoStoreActivity* GetActivity() const { return GetActivityTableChecked().GetActivity(GetRowIndex()); }
 		const FIoStoreActivity* GetActivityChecked() const { return GetActivityTableChecked().GetActivityChecked(GetRowIndex()); }
-	};
 
+	private:
+		uint32 RequestIndex;
+	};
 
 } //UE::IoStoreInsights
diff --git a/Engine/Plugins/IoStoreInsights/Source/IoStoreInsights/Private/Widgets/SIoStoreActivityTableTreeView.cpp b/Engine/Plugins/IoStoreInsights/Source/IoStoreInsights/Private/Widgets/SIoStoreActivityTableTreeView.cpp
--- a/Engine/Plugins/IoStoreInsights/Source/IoStoreInsights/Private/Widgets/SIoStoreActivityTableTreeView.cpp
+++ b/Engine/Plugins/IoStoreInsights/Source/IoStoreInsights/Private/Widgets/SIoStoreActivityTableTreeView.cpp
@@ -51,26 +51,29 @@ namespace UE::IoStoreInsights
 				const FIoStoreInsightsProvider* Provider = AnalysisSession->ReadProvider<FIoStoreInsightsProvider>(IIoStoreInsightsProvider::ProviderName);
 				if (Provider)
 				{
-					Provider->EnumerateIoStoreRequests([this, &ActivityTable, &Activities](const FIoStoreRequest& IoStoreRequest, const IIoStoreInsightsProvider::Timeline& Timeline)
-					{
-						Timeline.EnumerateEvents(StartTime, EndTime, [this, &ActivityTable, &Activities](double EventStartTime, double EventEndTime, uint32 EventDepth, const FIoStoreActivity* IoStoreActivity)
+					//TODO: thread safety!
+					Provider->EnumerateIoStoreRequests(
+						[this, &ActivityTable, &Activities]
+						(const FIoStoreRequest& IoStoreRequest, const IIoStoreInsightsProvider::Timeline& Timeline)
 						{
-							if (IoStoreActivity->ActivityType == EIoStoreActivityType::Request_Read && IoStoreActivity->EndTime > 0 && !IoStoreActivity->Failed)
-							{
-								static const FName BaseNodeName(TEXT("row"));
-
-								Activities.Emplace(IoStoreActivity);
-								uint32 Index = Activities.Num() - 1;
-								FName NodeName(BaseNodeName, static_cast<int32>(IoStoreActivity->IoStoreRequest->IoStoreRequestIndex+1));
-								FActivityNodePtr NodePtr = MakeShared<FIoStoreActivityNode>(NodeName, ActivityTable, Index);
-								TableRowNodes.Add(NodePtr);
-							}
-
-							return TraceServices::EEventEnumerate::Continue;
+							Timeline.EnumerateEvents(StartTime, EndTime,
+								[this, &ActivityTable, &Activities]
+								(double EventStartTime, double EventEndTime, uint32 EventDepth, const FIoStoreActivity* IoStoreActivity)
+								{
+									if (IoStoreActivity->ActivityType == EIoStoreActivityType::Request_Read && IoStoreActivity->EndTime > 0 && !IoStoreActivity->Failed)
+									{
+										Activities.Emplace(IoStoreActivity);
+										const uint32 TableRowIndex = Activities.Num() - 1;
+										const uint32 RequestIndex = IoStoreActivity->IoStoreRequest->IoStoreRequestIndex;
+										FActivityNodePtr NodePtr = MakeShared<FIoStoreActivityNode>(ActivityTable, TableRowIndex, RequestIndex);
+										TableRowNodes.Add(NodePtr);
+									}
+
+									return TraceServices::EEventEnumerate::Continue;
+								});
+
+							return true;
 						});
-
-						return true;
-					});
 				}
 			}
 
